Editor session summary now only writes the minimum amount of fields during a crash to cut down on memory usage.

Additionally, we now only save platform information (eg. hardware, OS versions) once per session.

diff --git a/Engine/Source/Developer/EditorAnalyticsSession/Private/EditorAnalyticsSession.cpp b/Engine/Source/Developer/EditorAnalyticsSession/Private/EditorAnalyticsSession.cpp
--- a/Engine/Source/Developer/EditorAnalyticsSession/Private/EditorAnalyticsSession.cpp
+++ b/Engine/Source/Developer/EditorAnalyticsSession/Private/EditorAnalyticsSession.cpp
@@ -248,6 +248,7 @@ FEditorAnalyticsSession::FEditorAnalyticsSession()
 	bIsInPIE = false;
 	bIsInEnterprise = false;
 	bIsInVRMode = false;
+	bAlreadySaved = false;
 }
 
 bool FEditorAnalyticsSession::Lock(FTimespan Timeout)
@@ -297,16 +298,46 @@ bool FEditorAnalyticsSession::Save()
 		return false;
 	}
 
-	FString StorageLocation = EditorAnalyticsUtils::GetSessionStorageLocation(SessionId);
+	const FString StorageLocation = EditorAnalyticsUtils::GetSessionStorageLocation(SessionId);
+
+	if (!bAlreadySaved)
+	{
+		SET_STORED_STRING(EngineVersion);
+		SET_STORED_INT(PlatformProcessID);
+		SET_STORED_STRING(DesktopGPUAdapter);
+		SET_STORED_STRING(RenderingGPUAdapter);
+		SET_STORED_INT(GPUVendorID);
+		SET_STORED_INT(GPUDeviceID);
+		SET_STORED_INT(GRHIDeviceRevision);
+		SET_STORED_STRING(GRHIAdapterInternalDriverVersion);
+		SET_STORED_STRING(GRHIAdapterUserDriverVersion);
+
+		FPlatformMisc::SetStoredValue(EditorAnalyticsDefs::StoreId, StorageLocation, EditorAnalyticsDefs::TotalPhysicalRAMStoreKey, FString::Printf(TEXT("%llu"), TotalPhysicalRAM));
+
+		SET_STORED_INT(CPUPhysicalCores);
+		SET_STORED_INT(CPULogicalCores);
+		SET_STORED_STRING(CPUVendor);
+		SET_STORED_STRING(CPUBrand);
+
+		FPlatformMisc::SetStoredValue(EditorAnalyticsDefs::StoreId, StorageLocation, EditorAnalyticsDefs::StartupTimestampStoreKey, EditorAnalyticsUtils::TimestampToString(StartupTimestamp));
+
+		SET_STORED_STRING(OSMajor);
+		SET_STORED_STRING(OSMinor);
+		SET_STORED_STRING(OSVersion);
+
+		FPlatformMisc::SetStoredValue(EditorAnalyticsDefs::StoreId, StorageLocation, EditorAnalyticsDefs::bIs64BitOSStoreKey, EditorAnalyticsUtils::BoolToStoredString(bIs64BitOS));
+
+		const FString PluginsString = FString::Join(Plugins, TEXT(","));
+		FPlatformMisc::SetStoredValue(EditorAnalyticsDefs::StoreId, StorageLocation, EditorAnalyticsDefs::PluginsStoreKey, PluginsString);
+
+		bAlreadySaved = true;
+	}
 
 	SET_STORED_STRING(ProjectName);
 	SET_STORED_STRING(ProjectID);
 	SET_STORED_STRING(ProjectDescription);
 	SET_STORED_STRING(ProjectVersion);
-	SET_STORED_STRING(EngineVersion);
-	SET_STORED_INT(PlatformProcessID);
 
-	FPlatformMisc::SetStoredValue(EditorAnalyticsDefs::StoreId, StorageLocation, EditorAnalyticsDefs::StartupTimestampStoreKey, EditorAnalyticsUtils::TimestampToString(StartupTimestamp));
 	FPlatformMisc::SetStoredValue(EditorAnalyticsDefs::StoreId, StorageLocation, EditorAnalyticsDefs::TimestampStoreKey, EditorAnalyticsUtils::TimestampToString(Timestamp));
 
 	SET_STORED_INT(Idle1Min);
@@ -315,31 +346,7 @@ bool FEditorAnalyticsSession::Save()
 
 	SET_STORED_STRING(CurrentUserActivity);
 
-	FString PluginsString = FString::Join(Plugins, TEXT(","));
-	FPlatformMisc::SetStoredValue(EditorAnalyticsDefs::StoreId, StorageLocation, EditorAnalyticsDefs::PluginsStoreKey, PluginsString);
-	
 	FPlatformMisc::SetStoredValue(EditorAnalyticsDefs::StoreId, StorageLocation, EditorAnalyticsDefs::AverageFPSStoreKey, FString::SanitizeFloat(AverageFPS));
-
-	SET_STORED_STRING(DesktopGPUAdapter);
-	SET_STORED_STRING(RenderingGPUAdapter);
-	SET_STORED_INT(GPUVendorID);
-	SET_STORED_INT(GPUDeviceID);
-	SET_STORED_INT(GRHIDeviceRevision);
-	SET_STORED_STRING(GRHIAdapterInternalDriverVersion);
-	SET_STORED_STRING(GRHIAdapterUserDriverVersion);
-
-	FPlatformMisc::SetStoredValue(EditorAnalyticsDefs::StoreId, StorageLocation, EditorAnalyticsDefs::TotalPhysicalRAMStoreKey, FString::Printf(TEXT("%llu"), TotalPhysicalRAM));
-
-	SET_STORED_INT(CPUPhysicalCores);
-	SET_STORED_INT(CPULogicalCores);
-	SET_STORED_STRING(CPUVendor);
-	SET_STORED_STRING(CPUBrand);
-
-	SET_STORED_STRING(OSMajor);
-	SET_STORED_STRING(OSMinor);
-	SET_STORED_STRING(OSVersion);
-
-	FPlatformMisc::SetStoredValue(EditorAnalyticsDefs::StoreId, StorageLocation, EditorAnalyticsDefs::bIs64BitOSStoreKey, EditorAnalyticsUtils::BoolToStoredString(bIs64BitOS));
 	FPlatformMisc::SetStoredValue(EditorAnalyticsDefs::StoreId, StorageLocation, EditorAnalyticsDefs::IsCrashStoreKey, EditorAnalyticsUtils::BoolToStoredString(bCrashed));
 	FPlatformMisc::SetStoredValue(EditorAnalyticsDefs::StoreId, StorageLocation, EditorAnalyticsDefs::IsGPUCrashStoreKey, EditorAnalyticsUtils::BoolToStoredString(bGPUCrashed));
 	FPlatformMisc::SetStoredValue(EditorAnalyticsDefs::StoreId, StorageLocation, EditorAnalyticsDefs::IsDebuggerStoreKey, EditorAnalyticsUtils::BoolToStoredString(bIsDebugger));
@@ -354,6 +361,24 @@ bool FEditorAnalyticsSession::Save()
 	return true;
 }
 
+bool FEditorAnalyticsSession::SaveForCrash()
+{
+	if (!ensure(IsLocked()))
+	{
+		return false;
+	}
+
+	const FString StorageLocation = EditorAnalyticsUtils::GetSessionStorageLocation(SessionId);
+	
+	FPlatformMisc::SetStoredValue(EditorAnalyticsDefs::StoreId, StorageLocation, EditorAnalyticsDefs::IsCrashStoreKey, EditorAnalyticsUtils::BoolToStoredString(bCrashed));
+	FPlatformMisc::SetStoredValue(EditorAnalyticsDefs::StoreId, StorageLocation, EditorAnalyticsDefs::IsGPUCrashStoreKey, EditorAnalyticsUtils::BoolToStoredString(bGPUCrashed));
+	FPlatformMisc::SetStoredValue(EditorAnalyticsDefs::StoreId, StorageLocation, EditorAnalyticsDefs::IsTerminatingKey, EditorAnalyticsUtils::BoolToStoredString(bIsTerminating));
+	FPlatformMisc::SetStoredValue(EditorAnalyticsDefs::StoreId, StorageLocation, EditorAnalyticsDefs::WasShutdownStoreKey, EditorAnalyticsUtils::BoolToStoredString(bWasShutdown));
+	FPlatformMisc::SetStoredValue(EditorAnalyticsDefs::StoreId, StorageLocation, EditorAnalyticsDefs::TimestampStoreKey, EditorAnalyticsUtils::TimestampToString(Timestamp));
+
+	return true;
+}
+
 #undef SET_STORED_INT
 #undef SET_STORED_STRING
 
@@ -452,7 +477,7 @@ bool FEditorAnalyticsSession::LoadAllStoredSessions(TArray<FEditorAnalyticsSessi
 		FEditorAnalyticsSession NewSession;
 		EditorAnalyticsUtils::LoadInternal(NewSession, Id);
 
-		OutSessions.Add(NewSession);
+		OutSessions.Add(MoveTemp(NewSession));
 	}
 
 	return true;
diff --git a/Engine/Source/Developer/EditorAnalyticsSession/Public/EditorAnalyticsSession.h b/Engine/Source/Developer/EditorAnalyticsSession/Public/EditorAnalyticsSession.h
--- a/Engine/Source/Developer/EditorAnalyticsSession/Public/EditorAnalyticsSession.h
+++ b/Engine/Source/Developer/EditorAnalyticsSession/Public/EditorAnalyticsSession.h
@@ -65,6 +65,12 @@ struct EDITORANALYTICSSESSION_API FEditorAnalyticsSession
 	 */
 	bool Save();
 
+	/**
+	 * Save only the bare minimum of fields that may have changed during a crash.
+	 * @returns true if the session was successfully saved.
+	 */
+	bool SaveForCrash();
+
 	/**
 	 *  Load a session with the given session ID from stored values.
 	 * @returns true if the session was found and successfully loaded.
@@ -113,6 +119,12 @@ struct EDITORANALYTICSSESSION_API FEditorAnalyticsSession
 private:
 
 	static FSystemWideCriticalSection* StoredValuesLock;
+
+	/** 
+	 * Has this session already been saved? 
+	 * If not, then the first save will write out session invariant details such as hardware specs.
+	 */
+	bool bAlreadySaved : 1;
 };
 
 class FEditorAnalyticsSessionModule : public IModuleInterface
diff --git a/Engine/Source/Runtime/Engine/Private/Analytics/EditorSessionSummaryWriter.cpp b/Engine/Source/Runtime/Engine/Private/Analytics/EditorSessionSummaryWriter.cpp
--- a/Engine/Source/Runtime/Engine/Private/Analytics/EditorSessionSummaryWriter.cpp
+++ b/Engine/Source/Runtime/Engine/Private/Analytics/EditorSessionSummaryWriter.cpp
@@ -188,7 +188,14 @@ void FEditorSessionSummaryWriter::Shutdown()
 
 		if (FEditorAnalyticsSession::Lock())
 		{
-			CurrentSession->Save();
+			if (CurrentSession->bWasShutdown)
+			{
+				CurrentSession->Save();
+			}
+			else
+			{
+				CurrentSession->SaveForCrash();
+			}
 
 			FEditorAnalyticsSession::Unlock();
 		}
@@ -273,7 +280,7 @@ void FEditorSessionSummaryWriter::OnCrashing()
 
 		if (FEditorAnalyticsSession::Lock())
 		{
-			CurrentSession->Save();
+			CurrentSession->SaveForCrash();
 			
 			FEditorAnalyticsSession::Unlock();
 		}
@@ -290,7 +297,7 @@ void FEditorSessionSummaryWriter::OnTerminate()
 
 		if (FEditorAnalyticsSession::Lock())
 		{
-			CurrentSession->Save();
+			CurrentSession->SaveForCrash();
 		
 			FEditorAnalyticsSession::Unlock();
 		}
