The list of currently-loaded plugins is now written to the editor session analytics payload.

diff --git a/Engine/Source/Developer/EditorAnalyticsSession/Private/EditorSessionSummarySender.cpp b/Engine/Source/Developer/EditorAnalyticsSession/Private/EditorSessionSummarySender.cpp
--- a/Engine/Source/Developer/EditorAnalyticsSession/Private/EditorSessionSummarySender.cpp
+++ b/Engine/Source/Developer/EditorAnalyticsSession/Private/EditorSessionSummarySender.cpp
@@ -179,6 +179,7 @@ void FEditorSessionSummarySender::SendSessionSummaryEvent(const FEditorAnalytics
 	AnalyticsAttributes.Emplace(TEXT("CurrentUserActivity"), Session.CurrentUserActivity);
 	AnalyticsAttributes.Emplace(TEXT("StartupTimestamp"), Session.StartupTimestamp.ToIso8601());
 	AnalyticsAttributes.Emplace(TEXT("AverageFPS"), Session.AverageFPS);
+	AnalyticsAttributes.Emplace(TEXT("Plugins"), PluginsString);
 
 	double SessionDuration = (Session.Timestamp - Session.StartupTimestamp).GetTotalSeconds();
 	AnalyticsAttributes.Emplace(TEXT("SessionDuration"), SessionDuration);
