EditorSessionSummary's session duration is now calculated using FPlatformTime::Seconds instead of FDateTime::UtcNow() timestamps, which are not guaranteed by platforms to increase monotonically, leading to outliers with negative session durations.

diff --git a/Engine/Source/Developer/EditorAnalyticsSession/Private/EditorAnalyticsSession.cpp b/Engine/Source/Developer/EditorAnalyticsSession/Private/EditorAnalyticsSession.cpp
--- a/Engine/Source/Developer/EditorAnalyticsSession/Private/EditorAnalyticsSession.cpp
+++ b/Engine/Source/Developer/EditorAnalyticsSession/Private/EditorAnalyticsSession.cpp
@@ -30,6 +30,7 @@ namespace EditorAnalyticsDefs
 	// timestamps
 	static const FString StartupTimestampStoreKey(TEXT("StartupTimestamp"));
 	static const FString TimestampStoreKey(TEXT("Timestamp"));
+	static const FString SessionDurationStoreKey(TEXT("SessionDuration"));
 	static const FString Idle1MinStoreKey(TEXT("Idle1Min"));
 	static const FString Idle5MinStoreKey(TEXT("Idle5Min"));
 	static const FString Idle30MinStoreKey(TEXT("Idle30Min"));
@@ -145,6 +146,7 @@ namespace EditorAnalyticsUtils
 			Session.Timestamp = EditorAnalyticsUtils::StringToTimestamp(TimestampString);
 		}
 
+		GET_STORED_INT(SessionDuration);
 		GET_STORED_INT(Idle1Min);
 		GET_STORED_INT(Idle5Min);
 		GET_STORED_INT(Idle30Min);
@@ -225,6 +227,7 @@ FEditorAnalyticsSession::FEditorAnalyticsSession()
 	PlatformProcessID = 0;
 	StartupTimestamp = FDateTime::MinValue();
 	Timestamp = FDateTime::MinValue();
+	SessionDuration = 0;
 	Idle1Min = 0;
 	Idle5Min = 0;
 	Idle30Min = 0;
@@ -340,6 +343,7 @@ bool FEditorAnalyticsSession::Save()
 
 	FPlatformMisc::SetStoredValue(EditorAnalyticsDefs::StoreId, StorageLocation, EditorAnalyticsDefs::TimestampStoreKey, EditorAnalyticsUtils::TimestampToString(Timestamp));
 
+	SET_STORED_INT(SessionDuration);
 	SET_STORED_INT(Idle1Min);
 	SET_STORED_INT(Idle5Min);
 	SET_STORED_INT(Idle30Min);
@@ -375,6 +379,7 @@ bool FEditorAnalyticsSession::SaveForCrash()
 	FPlatformMisc::SetStoredValue(EditorAnalyticsDefs::StoreId, StorageLocation, EditorAnalyticsDefs::IsTerminatingKey, EditorAnalyticsUtils::BoolToStoredString(bIsTerminating));
 	FPlatformMisc::SetStoredValue(EditorAnalyticsDefs::StoreId, StorageLocation, EditorAnalyticsDefs::WasShutdownStoreKey, EditorAnalyticsUtils::BoolToStoredString(bWasShutdown));
 	FPlatformMisc::SetStoredValue(EditorAnalyticsDefs::StoreId, StorageLocation, EditorAnalyticsDefs::TimestampStoreKey, EditorAnalyticsUtils::TimestampToString(Timestamp));
+	FPlatformMisc::SetStoredValue(EditorAnalyticsDefs::StoreId, StorageLocation, EditorAnalyticsDefs::SessionDurationStoreKey, FString::FromInt(SessionDuration));
 
 	return true;
 }
@@ -390,6 +395,8 @@ bool FEditorAnalyticsSession::Load(const FString& InSessionID)
 	}
 
 	EditorAnalyticsUtils::LoadInternal(*this, InSessionID);
+	bAlreadySaved = false;
+
 	return true;
 }
 
@@ -411,6 +418,7 @@ bool FEditorAnalyticsSession::Delete() const
 
 	FPlatformMisc::DeleteStoredValue(EditorAnalyticsDefs::StoreId, SectionName, EditorAnalyticsDefs::StartupTimestampStoreKey);
 	FPlatformMisc::DeleteStoredValue(EditorAnalyticsDefs::StoreId, SectionName, EditorAnalyticsDefs::TimestampStoreKey);
+	FPlatformMisc::DeleteStoredValue(EditorAnalyticsDefs::StoreId, SectionName, EditorAnalyticsDefs::SessionDurationStoreKey);
 	FPlatformMisc::DeleteStoredValue(EditorAnalyticsDefs::StoreId, SectionName, EditorAnalyticsDefs::Idle1MinStoreKey);
 	FPlatformMisc::DeleteStoredValue(EditorAnalyticsDefs::StoreId, SectionName, EditorAnalyticsDefs::Idle5MinStoreKey);
 	FPlatformMisc::DeleteStoredValue(EditorAnalyticsDefs::StoreId, SectionName, EditorAnalyticsDefs::Idle30MinStoreKey);
diff --git a/Engine/Source/Developer/EditorAnalyticsSession/Private/EditorSessionSummarySender.cpp b/Engine/Source/Developer/EditorAnalyticsSession/Private/EditorSessionSummarySender.cpp
--- a/Engine/Source/Developer/EditorAnalyticsSession/Private/EditorSessionSummarySender.cpp
+++ b/Engine/Source/Developer/EditorAnalyticsSession/Private/EditorSessionSummarySender.cpp
@@ -15,7 +15,7 @@ DEFINE_LOG_CATEGORY_STATIC(LogEditorSessionSummary, Verbose, All);
 namespace EditorSessionSenderDefs
 {
 	static const FTimespan SessionExpiration = FTimespan::FromDays(30.0);
-	static const int HeartbeatPeriodSeconds = 60;
+	static const float HeartbeatPeriodSeconds = 60;
 
 	// shutdown types
 	static const FString RunningSessionToken(TEXT("Running"));
@@ -42,7 +42,7 @@ void FEditorSessionSummarySender::Tick(float DeltaTime)
 {
 	HeartbeatTimeElapsed += DeltaTime;
 
-	if (HeartbeatTimeElapsed > (float) EditorSessionSenderDefs::HeartbeatPeriodSeconds)
+	if (HeartbeatTimeElapsed > EditorSessionSenderDefs::HeartbeatPeriodSeconds)
 	{
 		HeartbeatTimeElapsed = 0.0f;
 
@@ -175,18 +175,15 @@ void FEditorSessionSummarySender::SendSessionSummaryEvent(const FEditorAnalytics
 	AnalyticsAttributes.Emplace(TEXT("SessionId"), SessionIdString);
 	AnalyticsAttributes.Emplace(TEXT("EngineVersion"), Session.EngineVersion);
 	AnalyticsAttributes.Emplace(TEXT("ShutdownType"), ShutdownTypeString);
-	AnalyticsAttributes.Emplace(TEXT("Timestamp"), Session.Timestamp.ToIso8601());
-	AnalyticsAttributes.Emplace(TEXT("CurrentUserActivity"), Session.CurrentUserActivity);
 	AnalyticsAttributes.Emplace(TEXT("StartupTimestamp"), Session.StartupTimestamp.ToIso8601());
-	AnalyticsAttributes.Emplace(TEXT("AverageFPS"), Session.AverageFPS);
-	AnalyticsAttributes.Emplace(TEXT("Plugins"), PluginsString);
-
-	double SessionDuration = (Session.Timestamp - Session.StartupTimestamp).GetTotalSeconds();
-	AnalyticsAttributes.Emplace(TEXT("SessionDuration"), SessionDuration);
-
+	AnalyticsAttributes.Emplace(TEXT("Timestamp"), Session.Timestamp.ToIso8601());
+	AnalyticsAttributes.Emplace(TEXT("SessionDuration"), Session.SessionDuration);
 	AnalyticsAttributes.Emplace(TEXT("1MinIdle"), Session.Idle1Min);
 	AnalyticsAttributes.Emplace(TEXT("5MinIdle"), Session.Idle5Min);
 	AnalyticsAttributes.Emplace(TEXT("30MinIdle"), Session.Idle30Min);
+	AnalyticsAttributes.Emplace(TEXT("CurrentUserActivity"), Session.CurrentUserActivity);
+	AnalyticsAttributes.Emplace(TEXT("AverageFPS"), Session.AverageFPS);
+	AnalyticsAttributes.Emplace(TEXT("Plugins"), PluginsString);
 	AnalyticsAttributes.Emplace(TEXT("DesktopGPUAdapter"), Session.DesktopGPUAdapter);
 	AnalyticsAttributes.Emplace(TEXT("RenderingGPUAdapter"), Session.RenderingGPUAdapter);
 	AnalyticsAttributes.Emplace(TEXT("GPUVendorID"), Session.GPUVendorID);
@@ -203,7 +200,6 @@ void FEditorSessionSummarySender::SendSessionSummaryEvent(const FEditorAnalytics
 	AnalyticsAttributes.Emplace(TEXT("OSMinor"), Session.OSMinor);
 	AnalyticsAttributes.Emplace(TEXT("OSVersion"), Session.OSVersion);
 	AnalyticsAttributes.Emplace(TEXT("Is64BitOS"), Session.bIs64BitOS);
-
 	AnalyticsAttributes.Emplace(TEXT("GPUCrash"), Session.bGPUCrashed);
 	AnalyticsAttributes.Emplace(TEXT("WasDebugged"), Session.bWasEverDebugger);
 	AnalyticsAttributes.Emplace(TEXT("IsVanilla"), Session.bIsVanilla);
diff --git a/Engine/Source/Developer/EditorAnalyticsSession/Public/EditorAnalyticsSession.h b/Engine/Source/Developer/EditorAnalyticsSession/Public/EditorAnalyticsSession.h
--- a/Engine/Source/Developer/EditorAnalyticsSession/Public/EditorAnalyticsSession.h
+++ b/Engine/Source/Developer/EditorAnalyticsSession/Public/EditorAnalyticsSession.h
@@ -20,6 +20,7 @@ struct EDITORANALYTICSSESSION_API FEditorAnalyticsSession
 
 	FDateTime StartupTimestamp;
 	FDateTime Timestamp;
+	int32 SessionDuration;
 	int32 Idle1Min;
 	int32 Idle5Min;
 	int32 Idle30Min;
diff --git a/Engine/Source/Runtime/Engine/Private/Analytics/EditorSessionSummaryWriter.cpp b/Engine/Source/Runtime/Engine/Private/Analytics/EditorSessionSummaryWriter.cpp
--- a/Engine/Source/Runtime/Engine/Private/Analytics/EditorSessionSummaryWriter.cpp
+++ b/Engine/Source/Runtime/Engine/Private/Analytics/EditorSessionSummaryWriter.cpp
@@ -31,20 +31,18 @@
 
 DEFINE_LOG_CATEGORY_STATIC(LogEditorSessionSummary, Verbose, All);
 
-/* FEditorSessionSummaryWriter */
-
 namespace EditorSessionWriterDefs
 {
-	static const int HeartbeatPeriodSeconds = 60;
+	static const float HeartbeatPeriodSeconds = 60;
 }
 
 FEditorSessionSummaryWriter::FEditorSessionSummaryWriter() :
 	CurrentSession(nullptr)
+	, StartupSeconds(0.0)
 	, LastUserInteractionTime(0.0f)
 	, HeartbeatTimeElapsed(0.0f)
 	, bShutdown(false)
 {
-
 }
 
 void FEditorSessionSummaryWriter::Initialize()
@@ -56,6 +54,8 @@ void FEditorSessionSummaryWriter::Initialize()
 	FUserActivityTracking::OnActivityChanged.AddRaw(this, &FEditorSessionSummaryWriter::OnUserActivity);
 	FSlateApplication::Get().GetOnModalLoopTickEvent().AddRaw(this, &FEditorSessionSummaryWriter::Tick);
 
+	StartupSeconds = FPlatformTime::Seconds();
+
 	InitializeSessions();
 }
 
@@ -92,7 +92,10 @@ void FEditorSessionSummaryWriter::UpdateTimestamps()
 {
 	CurrentSession->Timestamp = FDateTime::UtcNow();
 
-	const float IdleSeconds = FPlatformTime::Seconds() - LastUserInteractionTime;
+	const double CurrentSeconds = FPlatformTime::Seconds();
+	CurrentSession->SessionDuration = FMath::FloorToInt(CurrentSeconds - StartupSeconds);
+
+	const double IdleSeconds = CurrentSeconds - LastUserInteractionTime;
 
 	// 1 + 1 minutes
 	if (IdleSeconds > (60 + 60))
@@ -125,7 +128,7 @@ void FEditorSessionSummaryWriter::Tick(float DeltaTime)
 
 	HeartbeatTimeElapsed += DeltaTime;
 
-	if (HeartbeatTimeElapsed > (float) EditorSessionWriterDefs::HeartbeatPeriodSeconds)
+	if (HeartbeatTimeElapsed > EditorSessionWriterDefs::HeartbeatPeriodSeconds)
 	{
 		HeartbeatTimeElapsed = 0.0f;
 
@@ -230,6 +233,7 @@ FEditorAnalyticsSession* FEditorSessionSummaryWriter::CreateCurrentSession() con
 	Session->ProjectVersion = ProjectSettings.ProjectVersion;
 	Session->EngineVersion = FEngineVersion::Current().ToString(EVersionComponent::Changelist);
 	Session->Timestamp = Session->StartupTimestamp = FDateTime::UtcNow();
+	Session->SessionDuration = 0;
 	Session->bIsDebugger = FPlatformMisc::IsDebuggerPresent();
 	Session->bWasEverDebugger = FPlatformMisc::IsDebuggerPresent();
 	Session->CurrentUserActivity = GetUserActivityString();
diff --git a/Engine/Source/Runtime/Engine/Private/Analytics/EditorSessionSummaryWriter.h b/Engine/Source/Runtime/Engine/Private/Analytics/EditorSessionSummaryWriter.h
--- a/Engine/Source/Runtime/Engine/Private/Analytics/EditorSessionSummaryWriter.h
+++ b/Engine/Source/Runtime/Engine/Private/Analytics/EditorSessionSummaryWriter.h
@@ -35,6 +35,7 @@ private:
 private:
 	FEditorAnalyticsSession* CurrentSession;
 	FString CurrentSessionSectionName;
+	double StartupSeconds;
 	double LastUserInteractionTime;
 	float HeartbeatTimeElapsed;
 	bool bShutdown;
