Disable over writing to ini files when using UNIX as it crashes the MallocCrash alloctor due to to many allocations

diff --git a/Engine/Source/Developer/EditorAnalyticsSession/Private/EditorAnalyticsSession.cpp b/Engine/Source/Developer/EditorAnalyticsSession/Private/EditorAnalyticsSession.cpp
--- a/Engine/Source/Developer/EditorAnalyticsSession/Private/EditorAnalyticsSession.cpp
+++ b/Engine/Source/Developer/EditorAnalyticsSession/Private/EditorAnalyticsSession.cpp
@@ -372,14 +372,18 @@ bool FEditorAnalyticsSession::SaveForCrash()
 		return false;
 	}
 
-	const FString StorageLocation = EditorAnalyticsUtils::GetSessionStorageLocation(SessionId);
-	
-	FPlatformMisc::SetStoredValue(EditorAnalyticsDefs::StoreId, StorageLocation, EditorAnalyticsDefs::IsCrashStoreKey, EditorAnalyticsUtils::BoolToStoredString(bCrashed));
-	FPlatformMisc::SetStoredValue(EditorAnalyticsDefs::StoreId, StorageLocation, EditorAnalyticsDefs::IsGPUCrashStoreKey, EditorAnalyticsUtils::BoolToStoredString(bGPUCrashed));
-	FPlatformMisc::SetStoredValue(EditorAnalyticsDefs::StoreId, StorageLocation, EditorAnalyticsDefs::IsTerminatingKey, EditorAnalyticsUtils::BoolToStoredString(bIsTerminating));
-	FPlatformMisc::SetStoredValue(EditorAnalyticsDefs::StoreId, StorageLocation, EditorAnalyticsDefs::WasShutdownStoreKey, EditorAnalyticsUtils::BoolToStoredString(bWasShutdown));
-	FPlatformMisc::SetStoredValue(EditorAnalyticsDefs::StoreId, StorageLocation, EditorAnalyticsDefs::TimestampStoreKey, EditorAnalyticsUtils::TimestampToString(Timestamp));
-	FPlatformMisc::SetStoredValue(EditorAnalyticsDefs::StoreId, StorageLocation, EditorAnalyticsDefs::SessionDurationStoreKey, FString::FromInt(SessionDuration));
+	// These ini writes are causing MallocCrash to go over its LARGE_MEMORYPOOL_SIZE due to writing to ini files causing a potentially a few large allocations per write
+	if (!PLATFORM_UNIX)
+	{
+		const FString StorageLocation = EditorAnalyticsUtils::GetSessionStorageLocation(SessionId);
+		
+		FPlatformMisc::SetStoredValue(EditorAnalyticsDefs::StoreId, StorageLocation, EditorAnalyticsDefs::IsCrashStoreKey, EditorAnalyticsUtils::BoolToStoredString(bCrashed));
+		FPlatformMisc::SetStoredValue(EditorAnalyticsDefs::StoreId, StorageLocation, EditorAnalyticsDefs::IsGPUCrashStoreKey, EditorAnalyticsUtils::BoolToStoredString(bGPUCrashed));
+		FPlatformMisc::SetStoredValue(EditorAnalyticsDefs::StoreId, StorageLocation, EditorAnalyticsDefs::IsTerminatingKey, EditorAnalyticsUtils::BoolToStoredString(bIsTerminating));
+		FPlatformMisc::SetStoredValue(EditorAnalyticsDefs::StoreId, StorageLocation, EditorAnalyticsDefs::WasShutdownStoreKey, EditorAnalyticsUtils::BoolToStoredString(bWasShutdown));
+		FPlatformMisc::SetStoredValue(EditorAnalyticsDefs::StoreId, StorageLocation, EditorAnalyticsDefs::TimestampStoreKey, EditorAnalyticsUtils::TimestampToString(Timestamp));
+		FPlatformMisc::SetStoredValue(EditorAnalyticsDefs::StoreId, StorageLocation, EditorAnalyticsDefs::SessionDurationStoreKey, FString::FromInt(SessionDuration));
+	}
 
 	return true;
 }
