Re-added the low drive space analytics field now that the SessionSummary analytics changes have been brought back into Dev-Editor.

Edigrated changes to analytics from Dev-Core, most notably FEditorSessionSummarySender::FindCurrentSession.

diff --git a/Engine/Source/Developer/EditorAnalyticsSession/Private/EditorAnalyticsSession.cpp b/Engine/Source/Developer/EditorAnalyticsSession/Private/EditorAnalyticsSession.cpp
--- a/Engine/Source/Developer/EditorAnalyticsSession/Private/EditorAnalyticsSession.cpp
+++ b/Engine/Source/Developer/EditorAnalyticsSession/Private/EditorAnalyticsSession.cpp
@@ -70,6 +70,7 @@ namespace EditorAnalyticsDefs
 	static const FString IsInPIEStoreKey(TEXT("IsInPIE"));
 	static const FString IsInEnterpriseStoreKey(TEXT("IsInEnterprise"));
 	static const FString IsInVRModeStoreKey(TEXT("IsInVRMode"));
+	static const FString IsLowDriveSpaceStoreKey(TEXT("IsLowDriveSpace"));
 }
 
 // Utilities for writing to stored values
@@ -198,8 +199,9 @@ namespace EditorAnalyticsUtils
 		Session.bIsTerminating = EditorAnalyticsUtils::GetStoredBool(SectionName, EditorAnalyticsDefs::IsTerminatingKey);
 		Session.bWasShutdown = EditorAnalyticsUtils::GetStoredBool(SectionName, EditorAnalyticsDefs::WasShutdownStoreKey);
 		Session.bIsInPIE = EditorAnalyticsUtils::GetStoredBool(SectionName, EditorAnalyticsDefs::IsInPIEStoreKey);
-		Session.bIsInVRMode = EditorAnalyticsUtils::GetStoredBool(SectionName, EditorAnalyticsDefs::IsInVRModeStoreKey);
 		Session.bIsInEnterprise = EditorAnalyticsUtils::GetStoredBool(SectionName, EditorAnalyticsDefs::IsInEnterpriseStoreKey);
+		Session.bIsInVRMode = EditorAnalyticsUtils::GetStoredBool(SectionName, EditorAnalyticsDefs::IsInVRModeStoreKey);
+		Session.bIsLowDriveSpace = EditorAnalyticsUtils::GetStoredBool(SectionName, EditorAnalyticsDefs::IsLowDriveSpaceStoreKey);
 	}
 
 #undef GET_STORED_INT
@@ -246,8 +248,9 @@ FEditorAnalyticsSession::FEditorAnalyticsSession()
 	bIsTerminating = false;
 	bWasShutdown = false;
 	bIsInPIE = false;
-	bIsInEnterprise = false;
 	bIsInVRMode = false;
+	bIsInEnterprise = false;
+	bIsLowDriveSpace = false;
 }
 
 bool FEditorAnalyticsSession::Lock(FTimespan Timeout)
@@ -350,6 +353,7 @@ bool FEditorAnalyticsSession::Save()
 	FPlatformMisc::SetStoredValue(EditorAnalyticsDefs::StoreId, StorageLocation, EditorAnalyticsDefs::IsInPIEStoreKey, EditorAnalyticsUtils::BoolToStoredString(bIsInPIE));
 	FPlatformMisc::SetStoredValue(EditorAnalyticsDefs::StoreId, StorageLocation, EditorAnalyticsDefs::IsInEnterpriseStoreKey, EditorAnalyticsUtils::BoolToStoredString(bIsInEnterprise));
 	FPlatformMisc::SetStoredValue(EditorAnalyticsDefs::StoreId, StorageLocation, EditorAnalyticsDefs::IsInVRModeStoreKey, EditorAnalyticsUtils::BoolToStoredString(bIsInVRMode));
+	FPlatformMisc::SetStoredValue(EditorAnalyticsDefs::StoreId, StorageLocation, EditorAnalyticsDefs::IsLowDriveSpaceStoreKey, EditorAnalyticsUtils::BoolToStoredString(bIsLowDriveSpace));
 
 	return true;
 }
@@ -422,6 +426,7 @@ bool FEditorAnalyticsSession::Delete() const
 	FPlatformMisc::DeleteStoredValue(EditorAnalyticsDefs::StoreId, SectionName, EditorAnalyticsDefs::IsInPIEStoreKey);
 	FPlatformMisc::DeleteStoredValue(EditorAnalyticsDefs::StoreId, SectionName, EditorAnalyticsDefs::IsInEnterpriseStoreKey);
 	FPlatformMisc::DeleteStoredValue(EditorAnalyticsDefs::StoreId, SectionName, EditorAnalyticsDefs::IsInVRModeStoreKey);
+	FPlatformMisc::DeleteStoredValue(EditorAnalyticsDefs::StoreId, SectionName, EditorAnalyticsDefs::IsLowDriveSpaceStoreKey);
 
 	return true;
 }
diff --git a/Engine/Source/Developer/EditorAnalyticsSession/Private/EditorSessionSummarySender.cpp b/Engine/Source/Developer/EditorAnalyticsSession/Private/EditorSessionSummarySender.cpp
--- a/Engine/Source/Developer/EditorAnalyticsSession/Private/EditorSessionSummarySender.cpp
+++ b/Engine/Source/Developer/EditorAnalyticsSession/Private/EditorSessionSummarySender.cpp
@@ -32,7 +32,6 @@ FEditorSessionSummarySender::FEditorSessionSummarySender(IAnalyticsProvider& InA
 	, Sender(InSenderName)
 	, CurrentSessionProcessId(InCurrentSessionProcessId)
 {
-	SendStoredSessions();
 }
 
 FEditorSessionSummarySender::~FEditorSessionSummarySender()
@@ -62,6 +61,40 @@ void FEditorSessionSummarySender::SetCurrentSessionExitCode(const int32 InCurren
 	CurrentSessionExitCode = InExitCode;
 }
 
+bool FEditorSessionSummarySender::FindCurrentSession(FEditorAnalyticsSession& OutSession) const
+{
+	if (FPlatformProcess::IsApplicationRunning(CurrentSessionProcessId))
+	{
+		// still running, can't be abnormal termination
+		return false;
+	}
+
+	bool bFound = false;
+
+	if (FEditorAnalyticsSession::Lock(FTimespan::FromMilliseconds(100)))
+	{
+		TArray<FEditorAnalyticsSession> ExistingSessions;
+		FEditorAnalyticsSession::LoadAllStoredSessions(ExistingSessions);
+
+		const int32 ProcessID = CurrentSessionProcessId;
+		FEditorAnalyticsSession* CurrentSession = ExistingSessions.FindByPredicate(
+			[ProcessID](const FEditorAnalyticsSession& Session)
+			{
+				return Session.PlatformProcessID == ProcessID;
+			});
+
+		if (CurrentSession != nullptr)
+		{
+			OutSession = *CurrentSession;
+			bFound = true;
+		}
+
+		FEditorAnalyticsSession::Unlock();
+	}
+
+	return bFound;
+}
+
 void FEditorSessionSummarySender::SendStoredSessions(const bool bForceSendCurrentSession) const
 {
 	TArray<FEditorAnalyticsSession> SessionsToReport;
@@ -177,14 +210,18 @@ void FEditorSessionSummarySender::SendSessionSummaryEvent(const FEditorAnalytics
 	AnalyticsAttributes.Emplace(TEXT("IsInPIE"), Session.bIsInPIE);
 	AnalyticsAttributes.Emplace(TEXT("IsInEnterprise"), Session.bIsInEnterprise);
 	AnalyticsAttributes.Emplace(TEXT("IsInVRMode"), Session.bIsInVRMode);
+	AnalyticsAttributes.Emplace(TEXT("IsLowDriveSpace"), Session.bIsLowDriveSpace);
 	AnalyticsAttributes.Emplace(TEXT("SentFrom"), Sender);
 
-	AnalyticsProvider.RecordEvent(TEXT("SessionSummary"), AnalyticsAttributes);
+	// was this sent from some other process than itself or the out-of-process monitor for that run?
+	AnalyticsAttributes.Emplace(TEXT("DelayedSend"), Session.PlatformProcessID != CurrentSessionProcessId);
 
 	if (Session.PlatformProcessID == CurrentSessionProcessId && CurrentSessionExitCode.IsSet())
 	{
 		AnalyticsAttributes.Emplace(TEXT("ExitCode"), CurrentSessionExitCode.GetValue());
 	}
 
+	AnalyticsProvider.RecordEvent(TEXT("SessionSummary"), AnalyticsAttributes);
+
 	UE_LOG(LogEditorSessionSummary, Log, TEXT("EditorSessionSummary sent report. Type=%s, SessionId=%s"), *ShutdownTypeString, *SessionIdString);
 }
diff --git a/Engine/Source/Developer/EditorAnalyticsSession/Public/EditorAnalyticsSession.h b/Engine/Source/Developer/EditorAnalyticsSession/Public/EditorAnalyticsSession.h
--- a/Engine/Source/Developer/EditorAnalyticsSession/Public/EditorAnalyticsSession.h
+++ b/Engine/Source/Developer/EditorAnalyticsSession/Public/EditorAnalyticsSession.h
@@ -56,6 +56,7 @@ struct EDITORANALYTICSSESSION_API FEditorAnalyticsSession
 	bool bIsInPIE : 1;
 	bool bIsInEnterprise : 1;
 	bool bIsInVRMode : 1;
+	bool bIsLowDriveSpace : 1;
 
 	FEditorAnalyticsSession();
 
diff --git a/Engine/Source/Developer/EditorAnalyticsSession/Public/EditorSessionSummarySender.h b/Engine/Source/Developer/EditorAnalyticsSession/Public/EditorSessionSummarySender.h
--- a/Engine/Source/Developer/EditorAnalyticsSession/Public/EditorSessionSummarySender.h
+++ b/Engine/Source/Developer/EditorAnalyticsSession/Public/EditorSessionSummarySender.h
@@ -22,6 +22,7 @@ public:
 	void Shutdown();
 
 	void SetCurrentSessionExitCode(const int32 InCurrentSessionProcessId, const int32 InExitCode);
+	bool FindCurrentSession(FEditorAnalyticsSession& OutSession) const;
 
 private:
 	/** Send any stored Sessions. */
diff --git a/Engine/Source/Runtime/Engine/Private/Analytics/EditorSessionSummaryWriter.cpp b/Engine/Source/Runtime/Engine/Private/Analytics/EditorSessionSummaryWriter.cpp
--- a/Engine/Source/Runtime/Engine/Private/Analytics/EditorSessionSummaryWriter.cpp
+++ b/Engine/Source/Runtime/Engine/Private/Analytics/EditorSessionSummaryWriter.cpp
@@ -157,19 +157,16 @@ void FEditorSessionSummaryWriter::Tick(float DeltaTime)
 			CurrentSession->bIsInPIE = FPlayWorldCommandCallbacks::IsInPIE();
 #endif
 
-			if (FEditorAnalyticsSession::Lock())
-			{
-				CurrentSession->Save();
-
-				FEditorAnalyticsSession::Unlock();
-			}
+			TrySaveCurrentSession();
 		}
 	}
 }
 
 void FEditorSessionSummaryWriter::LowDriveSpaceDetected()
 {
-	
+	CurrentSession->bIsLowDriveSpace = true;
+
+	TrySaveCurrentSession();
 }
 
 void FEditorSessionSummaryWriter::Shutdown()
@@ -191,12 +188,7 @@ void FEditorSessionSummaryWriter::Shutdown()
 
 		UpdateTimestamps();
 
-		if (FEditorAnalyticsSession::Lock())
-		{
-			CurrentSession->Save();
-
-			FEditorAnalyticsSession::Unlock();
-		}
+		TrySaveCurrentSession();
 
 		delete CurrentSession;
 		CurrentSession = nullptr;
@@ -276,12 +268,7 @@ void FEditorSessionSummaryWriter::OnCrashing()
 		CurrentSession->bCrashed = true;
 		CurrentSession->bGPUCrashed = GIsGPUCrashed;
 
-		if (FEditorAnalyticsSession::Lock())
-		{
-			CurrentSession->Save();
-			
-			FEditorAnalyticsSession::Unlock();
-		}
+		TrySaveCurrentSession();
 	}
 }
 
@@ -293,12 +280,7 @@ void FEditorSessionSummaryWriter::OnTerminate()
 
 		CurrentSession->bIsTerminating = true;
 
-		if (FEditorAnalyticsSession::Lock())
-		{
-			CurrentSession->Save();
-		
-			FEditorAnalyticsSession::Unlock();
-		}
+		TrySaveCurrentSession();
 
 		if (IsEngineExitRequested())
 		{
@@ -313,12 +295,7 @@ void FEditorSessionSummaryWriter::OnVanillaStateChanged(bool bIsVanilla)
 	{
 		CurrentSession->bIsVanilla = bIsVanilla;
 
-		if (FEditorAnalyticsSession::Lock())
-		{
-			CurrentSession->Save();
-		
-			FEditorAnalyticsSession::Unlock();
-		}
+		TrySaveCurrentSession();
 	}
 }
 
@@ -328,12 +305,7 @@ void FEditorSessionSummaryWriter::OnUserActivity(const FUserActivity& UserActivi
 	{
 		CurrentSession->CurrentUserActivity = GetUserActivityString();
 
-		if (FEditorAnalyticsSession::Lock())
-		{
-			CurrentSession->Save();
-		
-			FEditorAnalyticsSession::Unlock();
-		}
+		TrySaveCurrentSession();
 	}
 }
 
@@ -349,5 +321,15 @@ FString FEditorSessionSummaryWriter::GetUserActivityString() const
 	return UserActivity.ActionName;
 }
 
+void FEditorSessionSummaryWriter::TrySaveCurrentSession()
+{
+	if (FEditorAnalyticsSession::Lock())
+	{
+		CurrentSession->Save();
+
+		FEditorAnalyticsSession::Unlock();
+	}
+}
+
 #undef LOCTEXT_NAMESPACE
 #endif // WITH_EDITOR
diff --git a/Engine/Source/Runtime/Engine/Private/Analytics/EditorSessionSummaryWriter.h b/Engine/Source/Runtime/Engine/Private/Analytics/EditorSessionSummaryWriter.h
--- a/Engine/Source/Runtime/Engine/Private/Analytics/EditorSessionSummaryWriter.h
+++ b/Engine/Source/Runtime/Engine/Private/Analytics/EditorSessionSummaryWriter.h
@@ -32,6 +32,7 @@ private:
 
 	FString GetUserActivityString() const;
 	void UpdateTimestamps();
+	void TrySaveCurrentSession();
 
 private:
 	FEditorAnalyticsSession* CurrentSession;
