SessionSummary analytics event is now sent from the out-of-process CrashReportClient. Currently communication is done through the local storage API in FPlatformMisc. As a result, many analytics fields that were being resolved only when sending analytics are now resolved on startup and stored temporarily.

EditorSessionSummarySender and EditorAnalyticsSession now live in their own module so they can be instantiated from the CrashReportClient.

diff --git a/Engine/Source/Developer/EditorAnalyticsSession/EditorAnalyticsSession.Build.cs b/Engine/Source/Developer/EditorAnalyticsSession/EditorAnalyticsSession.Build.cs
new file mode 100644
--- /dev/null
+++ b/Engine/Source/Developer/EditorAnalyticsSession/EditorAnalyticsSession.Build.cs
@@ -0,0 +1,18 @@
+// Copyright 1998-2019 Epic Games, Inc. All Rights Reserved.
+
+namespace UnrealBuildTool.Rules
+{
+	public class EditorAnalyticsSession : ModuleRules
+	{
+		public EditorAnalyticsSession( ReadOnlyTargetRules Target ) : base( Target )
+		{
+			PrivateDependencyModuleNames.AddRange(
+				new string[]
+				{
+					"Core",
+					"Analytics"
+				}
+			);
+		}
+	}
+}
diff --git a/Engine/Source/Developer/EditorAnalyticsSession/Private/EditorAnalyticsSession.cpp b/Engine/Source/Developer/EditorAnalyticsSession/Private/EditorAnalyticsSession.cpp
new file mode 100644
--- /dev/null
+++ b/Engine/Source/Developer/EditorAnalyticsSession/Private/EditorAnalyticsSession.cpp
@@ -0,0 +1,479 @@
+#include "EditorAnalyticsSession.h"
+
+IMPLEMENT_MODULE(FEditorAnalyticsSessionModule, EditorAnalyticsSession);
+
+namespace EditorAnalyticsDefs
+{
+	static const FString FalseValueString(TEXT("0"));
+	static const FString TrueValueString(TEXT("1"));
+
+	static const FString DefaultUserActivity(TEXT("Unknown"));
+	static const FString UnknownProjectValueString(TEXT("UnknownProject"));
+
+	// storage location
+	static const FString StoreId(TEXT("Epic Games"));
+	static const FString SessionSummarySection(TEXT("Unreal Engine/Session Summary/1_0"));
+	static const FString GlobalLockName(TEXT("UE4_SessionSummary_Lock"));
+	static const FString SessionListStoreKey(TEXT("SessionList"));
+
+	// general values
+	static const FString ProjectNameStoreKey(TEXT("ProjectName"));
+	static const FString ProjectIDStoreKey(TEXT("ProjectID"));
+	static const FString ProjectDescriptionStoreKey(TEXT("ProjectDescription"));
+	static const FString ProjectVersionStoreKey(TEXT("ProjectVersion"));
+	static const FString EngineVersionStoreKey(TEXT("EngineVersion"));
+	static const FString PlatformProcessIDStoreKey(TEXT("PlatformProcessID"));
+
+	// timestamps
+	static const FString StartupTimestampStoreKey(TEXT("StartupTimestamp"));
+	static const FString TimestampStoreKey(TEXT("Timestamp"));
+	static const FString Idle1MinStoreKey(TEXT("Idle1Min"));
+	static const FString Idle5MinStoreKey(TEXT("Idle5Min"));
+	static const FString Idle30MinStoreKey(TEXT("Idle30Min"));
+	static const FString CurrentUserActivityStoreKey(TEXT("CurrentUserActivity"));
+	static const FString PluginsStoreKey(TEXT("Plugins"));
+	static const FString AverageFPSStoreKey(TEXT("AverageFPS"));
+
+	// GPU details
+	static const FString DesktopGPUAdapterStoreKey(TEXT("DesktopGPUAdapter"));
+	static const FString RenderingGPUAdapterStoreKey(TEXT("RenderingGPUAdapter"));
+	static const FString GPUVendorIDStoreKey(TEXT("GPUVendorID"));
+	static const FString GPUDeviceIDStoreKey(TEXT("GPUDeviceID"));
+	static const FString GRHIDeviceRevisionStoreKey(TEXT("GRHIDeviceRevision"));
+	static const FString GRHIAdapterInternalDriverVersionStoreKey(TEXT("GRHIAdapterInternalDriverVersion"));
+	static const FString GRHIAdapterUserDriverVersionStoreKey(TEXT("GRHIAdapterUserDriverVersion"));
+
+	// CPU details
+	static const FString TotalPhysicalRAMStoreKey(TEXT("TotalPhysicalRAM"));
+	static const FString CPUPhysicalCoresStoreKey(TEXT("CPUPhysicalCores"));
+	static const FString CPULogicalCoresStoreKey(TEXT("CPULogicalCores"));
+	static const FString CPUVendorStoreKey(TEXT("CPUVendor"));
+	static const FString CPUBrandStoreKey(TEXT("CPUBrand"));
+
+	// OS details
+	static const FString OSMajorStoreKey(TEXT("OSMajor"));
+	static const FString OSMinorStoreKey(TEXT("OSMinor"));
+	static const FString OSVersionStoreKey(TEXT("OSVersion"));
+	static const FString bIs64BitOSStoreKey(TEXT("bIs64BitOS"));
+
+	// boolean flags
+	static const FString IsCrashStoreKey(TEXT("IsCrash"));
+	static const FString IsGPUCrashStoreKey(TEXT("IsGPUCrash"));
+	static const FString IsDebuggerStoreKey(TEXT("IsDebugger"));
+	static const FString WasDebuggerStoreKey(TEXT("WasEverDebugger"));
+	static const FString IsVanillaStoreKey(TEXT("IsVanilla"));
+	static const FString IsTerminatingKey(TEXT("Terminating"));
+	static const FString WasShutdownStoreKey(TEXT("WasShutdown"));
+	static const FString IsInPIEStoreKey(TEXT("IsInPIE"));
+	static const FString IsInEnterpriseStoreKey(TEXT("IsInEnterprise"));
+	static const FString IsInVRModeStoreKey(TEXT("IsInVRMode"));
+}
+
+// Utilities for writing to stored values
+namespace EditorAnalyticsUtils
+{
+	static FString TimestampToString(FDateTime InTimestamp)
+	{
+		return LexToString(InTimestamp.ToUnixTimestamp());
+	}
+
+	static FDateTime StringToTimestamp(FString InString)
+	{
+		int64 TimestampUnix;
+		if (LexTryParseString(TimestampUnix, *InString))
+		{
+			return FDateTime::FromUnixTimestamp(TimestampUnix);
+		}
+		return FDateTime::MinValue();
+	}
+
+	static FString BoolToStoredString(bool bInValue)
+	{
+		return bInValue ? EditorAnalyticsDefs::TrueValueString : EditorAnalyticsDefs::FalseValueString;
+	}
+
+	static bool GetStoredBool(const FString& SectionName, const FString& StoredKey)
+	{
+		FString StoredString = EditorAnalyticsDefs::FalseValueString;
+		FPlatformMisc::GetStoredValue(EditorAnalyticsDefs::StoreId, SectionName, StoredKey, StoredString);
+
+		return StoredString == EditorAnalyticsDefs::TrueValueString;
+	}
+
+	static FString GetSessionStorageLocation(const FString& SessionID)
+	{
+		return EditorAnalyticsDefs::SessionSummarySection + TEXT("/") + SessionID;
+	}
+
+// Utility macros to make it easier to check that all fields are being written to.
+#define GET_STORED_STRING(FieldName) FPlatformMisc::GetStoredValue(EditorAnalyticsDefs::StoreId, SectionName, EditorAnalyticsDefs:: FieldName ## StoreKey, Session.FieldName)
+#define GET_STORED_INT(FieldName) \
+	{ \
+		FString FieldName ## Temp; \
+		FPlatformMisc::GetStoredValue(EditorAnalyticsDefs::StoreId, SectionName, EditorAnalyticsDefs:: FieldName ## StoreKey, FieldName ## Temp); \
+		Session.FieldName = FCString::Atoi(*FieldName ## Temp); \
+	}
+	
+
+	static void LoadInternal(FEditorAnalyticsSession& Session, const FString& InSessionId)
+	{
+		Session.SessionId = InSessionId;
+
+		FString SectionName = EditorAnalyticsUtils::GetSessionStorageLocation(Session.SessionId);
+
+		GET_STORED_STRING(ProjectName);
+		GET_STORED_STRING(ProjectID);
+		GET_STORED_STRING(ProjectDescription);
+		GET_STORED_STRING(ProjectVersion);
+		GET_STORED_STRING(EngineVersion);
+
+		GET_STORED_INT(PlatformProcessID);
+
+		// scope is just to isolate the temporary value
+		{
+			FString StartupTimestampString;
+			FPlatformMisc::GetStoredValue(EditorAnalyticsDefs::StoreId, SectionName, EditorAnalyticsDefs::StartupTimestampStoreKey, StartupTimestampString);
+			Session.StartupTimestamp = EditorAnalyticsUtils::StringToTimestamp(StartupTimestampString);
+		}
+
+		{
+			FString TimestampString;
+			FPlatformMisc::GetStoredValue(EditorAnalyticsDefs::StoreId, SectionName, EditorAnalyticsDefs::TimestampStoreKey, TimestampString);
+			Session.Timestamp = EditorAnalyticsUtils::StringToTimestamp(TimestampString);
+		}
+
+		GET_STORED_INT(Idle1Min);
+		GET_STORED_INT(Idle5Min);
+		GET_STORED_INT(Idle30Min);
+
+		GET_STORED_STRING(CurrentUserActivity);
+
+		{
+			FString PluginsString;
+			FPlatformMisc::GetStoredValue(EditorAnalyticsDefs::StoreId, SectionName, EditorAnalyticsDefs::PluginsStoreKey, PluginsString);
+			PluginsString.ParseIntoArray(Session.Plugins, TEXT(","));
+		}
+
+		{
+			FString AverageFPSString;
+			FPlatformMisc::GetStoredValue(EditorAnalyticsDefs::StoreId, SectionName, EditorAnalyticsDefs::AverageFPSStoreKey, AverageFPSString);
+			Session.AverageFPS = FCString::Atof(*AverageFPSString);
+		}
+
+		GET_STORED_STRING(DesktopGPUAdapter);
+		GET_STORED_STRING(RenderingGPUAdapter);
+
+		GET_STORED_INT(GPUVendorID);
+		GET_STORED_INT(GPUDeviceID);
+		GET_STORED_INT(GRHIDeviceRevision);
+
+		GET_STORED_STRING(GRHIAdapterInternalDriverVersion);
+		GET_STORED_STRING(GRHIAdapterUserDriverVersion);
+
+		{
+			FString TotalPhysicalRAMString;
+			FPlatformMisc::GetStoredValue(EditorAnalyticsDefs::StoreId, SectionName, EditorAnalyticsDefs::GRHIAdapterUserDriverVersionStoreKey, TotalPhysicalRAMString);
+			Session.TotalPhysicalRAM = FCString::Atoi64(*TotalPhysicalRAMString);
+		}
+
+		GET_STORED_INT(CPUPhysicalCores);
+		GET_STORED_INT(CPULogicalCores);
+
+		GET_STORED_STRING(CPUVendor);
+		GET_STORED_STRING(CPUBrand);
+
+		GET_STORED_STRING(OSMajor);
+		GET_STORED_STRING(OSMinor);
+		GET_STORED_STRING(OSVersion);
+
+		Session.bIs64BitOS = EditorAnalyticsUtils::GetStoredBool(SectionName, EditorAnalyticsDefs::bIs64BitOSStoreKey);
+		Session.bCrashed = EditorAnalyticsUtils::GetStoredBool(SectionName, EditorAnalyticsDefs::IsCrashStoreKey);
+		Session.bGPUCrashed = EditorAnalyticsUtils::GetStoredBool(SectionName, EditorAnalyticsDefs::IsGPUCrashStoreKey);
+		Session.bIsDebugger = EditorAnalyticsUtils::GetStoredBool(SectionName, EditorAnalyticsDefs::IsDebuggerStoreKey);
+		Session.bWasEverDebugger = EditorAnalyticsUtils::GetStoredBool(SectionName, EditorAnalyticsDefs::WasDebuggerStoreKey);
+		Session.bIsVanilla = EditorAnalyticsUtils::GetStoredBool(SectionName, EditorAnalyticsDefs::IsVanillaStoreKey);
+		Session.bIsTerminating = EditorAnalyticsUtils::GetStoredBool(SectionName, EditorAnalyticsDefs::IsTerminatingKey);
+		Session.bWasShutdown = EditorAnalyticsUtils::GetStoredBool(SectionName, EditorAnalyticsDefs::WasShutdownStoreKey);
+		Session.bIsInPIE = EditorAnalyticsUtils::GetStoredBool(SectionName, EditorAnalyticsDefs::IsInPIEStoreKey);
+		Session.bIsInVRMode = EditorAnalyticsUtils::GetStoredBool(SectionName, EditorAnalyticsDefs::IsInVRModeStoreKey);
+		Session.bIsInEnterprise = EditorAnalyticsUtils::GetStoredBool(SectionName, EditorAnalyticsDefs::IsInEnterpriseStoreKey);
+	}
+
+#undef GET_STORED_INT
+#undef GET_STORED_STRING
+
+	static TArray<FString> GetSessionList()
+	{
+		FString SessionListString;
+		FPlatformMisc::GetStoredValue(EditorAnalyticsDefs::StoreId, EditorAnalyticsDefs::SessionSummarySection, EditorAnalyticsDefs::SessionListStoreKey, SessionListString);
+
+		TArray<FString> SessionIDs;
+		SessionListString.ParseIntoArray(SessionIDs, TEXT(","));
+
+		return MoveTemp(SessionIDs);
+	}
+}
+
+FSystemWideCriticalSection* FEditorAnalyticsSession::StoredValuesLock = nullptr;
+
+FEditorAnalyticsSession::FEditorAnalyticsSession()
+{
+	ProjectName = EditorAnalyticsDefs::UnknownProjectValueString;
+	PlatformProcessID = 0;
+	StartupTimestamp = FDateTime::MinValue();
+	Timestamp = FDateTime::MinValue();
+	Idle1Min = 0;
+	Idle5Min = 0;
+	Idle30Min = 0;
+	CurrentUserActivity = EditorAnalyticsDefs::DefaultUserActivity;
+	AverageFPS = 0;
+	GPUVendorID = 0;
+	GPUDeviceID = 0;
+	GRHIDeviceRevision = 0;
+	TotalPhysicalRAM = 0;
+	CPUPhysicalCores = 0;
+	CPULogicalCores = 0;
+
+	bIs64BitOS = false;
+	bCrashed = false;
+	bGPUCrashed = false;
+	bIsDebugger = false;
+	bWasEverDebugger = false;
+	bIsVanilla = false;
+	bIsTerminating = false;
+	bWasShutdown = false;
+	bIsInPIE = false;
+	bIsInEnterprise = false;
+	bIsInVRMode = false;
+}
+
+bool FEditorAnalyticsSession::Lock(FTimespan Timeout)
+{
+	if (!ensure(!IsLocked()))
+	{
+		return true;
+	}
+	
+	StoredValuesLock = new FSystemWideCriticalSection(EditorAnalyticsDefs::GlobalLockName, Timeout);
+
+	if (!IsLocked())
+	{
+		delete StoredValuesLock;
+		StoredValuesLock = nullptr;
+
+		return false;
+	}
+
+	return true;
+}
+
+void FEditorAnalyticsSession::Unlock()
+{
+	if (!ensure(IsLocked()))
+	{
+		return;
+	}
+
+	delete StoredValuesLock;
+	StoredValuesLock = nullptr;
+}
+
+bool FEditorAnalyticsSession::IsLocked()
+{
+	return StoredValuesLock != nullptr && StoredValuesLock->IsValid();
+}
+
+// Utility macros to make it easier to see that all values are correctly set
+#define SET_STORED_STRING(FieldName) FPlatformMisc::SetStoredValue(EditorAnalyticsDefs::StoreId, StorageLocation, EditorAnalyticsDefs::FieldName ## StoreKey, FieldName)
+#define SET_STORED_INT(FieldName) FPlatformMisc::SetStoredValue(EditorAnalyticsDefs::StoreId, StorageLocation, EditorAnalyticsDefs::FieldName ## StoreKey, FString::FromInt(FieldName))
+
+bool FEditorAnalyticsSession::Save()
+{
+	if (!ensure(IsLocked()))
+	{
+		return false;
+	}
+
+	FString StorageLocation = EditorAnalyticsUtils::GetSessionStorageLocation(SessionId);
+
+	SET_STORED_STRING(ProjectName);
+	SET_STORED_STRING(ProjectID);
+	SET_STORED_STRING(ProjectDescription);
+	SET_STORED_STRING(ProjectVersion);
+	SET_STORED_STRING(EngineVersion);
+	SET_STORED_INT(PlatformProcessID);
+
+	FPlatformMisc::SetStoredValue(EditorAnalyticsDefs::StoreId, StorageLocation, EditorAnalyticsDefs::StartupTimestampStoreKey, EditorAnalyticsUtils::TimestampToString(StartupTimestamp));
+	FPlatformMisc::SetStoredValue(EditorAnalyticsDefs::StoreId, StorageLocation, EditorAnalyticsDefs::TimestampStoreKey, EditorAnalyticsUtils::TimestampToString(Timestamp));
+
+	SET_STORED_INT(Idle1Min);
+	SET_STORED_INT(Idle5Min);
+	SET_STORED_INT(Idle30Min);
+
+	SET_STORED_STRING(CurrentUserActivity);
+
+	FString PluginsString = FString::Join(Plugins, TEXT(","));
+	FPlatformMisc::SetStoredValue(EditorAnalyticsDefs::StoreId, StorageLocation, EditorAnalyticsDefs::PluginsStoreKey, PluginsString);
+	
+	FPlatformMisc::SetStoredValue(EditorAnalyticsDefs::StoreId, StorageLocation, EditorAnalyticsDefs::AverageFPSStoreKey, FString::SanitizeFloat(AverageFPS));
+
+	SET_STORED_STRING(DesktopGPUAdapter);
+	SET_STORED_STRING(RenderingGPUAdapter);
+	SET_STORED_INT(GPUVendorID);
+	SET_STORED_INT(GPUDeviceID);
+	SET_STORED_INT(GRHIDeviceRevision);
+	SET_STORED_STRING(GRHIAdapterInternalDriverVersion);
+	SET_STORED_STRING(GRHIAdapterUserDriverVersion);
+
+	FPlatformMisc::SetStoredValue(EditorAnalyticsDefs::StoreId, StorageLocation, EditorAnalyticsDefs::TotalPhysicalRAMStoreKey, FString::Printf(TEXT("%llu"), TotalPhysicalRAM));
+
+	SET_STORED_INT(CPUPhysicalCores);
+	SET_STORED_INT(CPULogicalCores);
+	SET_STORED_STRING(CPUVendor);
+	SET_STORED_STRING(CPUBrand);
+
+	SET_STORED_STRING(OSMajor);
+	SET_STORED_STRING(OSMinor);
+	SET_STORED_STRING(OSVersion);
+
+	FPlatformMisc::SetStoredValue(EditorAnalyticsDefs::StoreId, StorageLocation, EditorAnalyticsDefs::bIs64BitOSStoreKey, EditorAnalyticsUtils::BoolToStoredString(bIs64BitOS));
+	FPlatformMisc::SetStoredValue(EditorAnalyticsDefs::StoreId, StorageLocation, EditorAnalyticsDefs::IsCrashStoreKey, EditorAnalyticsUtils::BoolToStoredString(bCrashed));
+	FPlatformMisc::SetStoredValue(EditorAnalyticsDefs::StoreId, StorageLocation, EditorAnalyticsDefs::IsGPUCrashStoreKey, EditorAnalyticsUtils::BoolToStoredString(bGPUCrashed));
+	FPlatformMisc::SetStoredValue(EditorAnalyticsDefs::StoreId, StorageLocation, EditorAnalyticsDefs::IsDebuggerStoreKey, EditorAnalyticsUtils::BoolToStoredString(bIsDebugger));
+	FPlatformMisc::SetStoredValue(EditorAnalyticsDefs::StoreId, StorageLocation, EditorAnalyticsDefs::WasDebuggerStoreKey, EditorAnalyticsUtils::BoolToStoredString(bWasEverDebugger));
+	FPlatformMisc::SetStoredValue(EditorAnalyticsDefs::StoreId, StorageLocation, EditorAnalyticsDefs::IsVanillaStoreKey, EditorAnalyticsUtils::BoolToStoredString(bIsVanilla));
+	FPlatformMisc::SetStoredValue(EditorAnalyticsDefs::StoreId, StorageLocation, EditorAnalyticsDefs::IsTerminatingKey, EditorAnalyticsUtils::BoolToStoredString(bIsTerminating));
+	FPlatformMisc::SetStoredValue(EditorAnalyticsDefs::StoreId, StorageLocation, EditorAnalyticsDefs::WasShutdownStoreKey, EditorAnalyticsUtils::BoolToStoredString(bWasShutdown));
+	FPlatformMisc::SetStoredValue(EditorAnalyticsDefs::StoreId, StorageLocation, EditorAnalyticsDefs::IsInPIEStoreKey, EditorAnalyticsUtils::BoolToStoredString(bIsInPIE));
+	FPlatformMisc::SetStoredValue(EditorAnalyticsDefs::StoreId, StorageLocation, EditorAnalyticsDefs::IsInEnterpriseStoreKey, EditorAnalyticsUtils::BoolToStoredString(bIsInEnterprise));
+	FPlatformMisc::SetStoredValue(EditorAnalyticsDefs::StoreId, StorageLocation, EditorAnalyticsDefs::IsInVRModeStoreKey, EditorAnalyticsUtils::BoolToStoredString(bIsInVRMode));
+
+	return true;
+}
+
+#undef SET_STORED_INT
+#undef SET_STORED_STRING
+
+bool FEditorAnalyticsSession::Load(const FString& InSessionID)
+{
+	if (!ensure(IsLocked()))
+	{
+		return false;
+	}
+
+	EditorAnalyticsUtils::LoadInternal(*this, InSessionID);
+	return true;
+}
+
+bool FEditorAnalyticsSession::Delete() const
+{
+	if (!ensure(IsLocked()))
+	{
+		return false;
+	}
+
+	FString SectionName = EditorAnalyticsUtils::GetSessionStorageLocation(SessionId);
+
+	FPlatformMisc::DeleteStoredValue(EditorAnalyticsDefs::StoreId, SectionName, EditorAnalyticsDefs::ProjectNameStoreKey);
+	FPlatformMisc::DeleteStoredValue(EditorAnalyticsDefs::StoreId, SectionName, EditorAnalyticsDefs::ProjectIDStoreKey);
+	FPlatformMisc::DeleteStoredValue(EditorAnalyticsDefs::StoreId, SectionName, EditorAnalyticsDefs::ProjectDescriptionStoreKey);
+	FPlatformMisc::DeleteStoredValue(EditorAnalyticsDefs::StoreId, SectionName, EditorAnalyticsDefs::ProjectVersionStoreKey);
+	FPlatformMisc::DeleteStoredValue(EditorAnalyticsDefs::StoreId, SectionName, EditorAnalyticsDefs::EngineVersionStoreKey);
+	FPlatformMisc::DeleteStoredValue(EditorAnalyticsDefs::StoreId, SectionName, EditorAnalyticsDefs::PlatformProcessIDStoreKey);
+
+	FPlatformMisc::DeleteStoredValue(EditorAnalyticsDefs::StoreId, SectionName, EditorAnalyticsDefs::StartupTimestampStoreKey);
+	FPlatformMisc::DeleteStoredValue(EditorAnalyticsDefs::StoreId, SectionName, EditorAnalyticsDefs::TimestampStoreKey);
+	FPlatformMisc::DeleteStoredValue(EditorAnalyticsDefs::StoreId, SectionName, EditorAnalyticsDefs::Idle1MinStoreKey);
+	FPlatformMisc::DeleteStoredValue(EditorAnalyticsDefs::StoreId, SectionName, EditorAnalyticsDefs::Idle5MinStoreKey);
+	FPlatformMisc::DeleteStoredValue(EditorAnalyticsDefs::StoreId, SectionName, EditorAnalyticsDefs::Idle30MinStoreKey);
+	FPlatformMisc::DeleteStoredValue(EditorAnalyticsDefs::StoreId, SectionName, EditorAnalyticsDefs::CurrentUserActivityStoreKey);
+	FPlatformMisc::DeleteStoredValue(EditorAnalyticsDefs::StoreId, SectionName, EditorAnalyticsDefs::PluginsStoreKey);
+	FPlatformMisc::DeleteStoredValue(EditorAnalyticsDefs::StoreId, SectionName, EditorAnalyticsDefs::AverageFPSStoreKey);
+
+	FPlatformMisc::DeleteStoredValue(EditorAnalyticsDefs::StoreId, SectionName, EditorAnalyticsDefs::DesktopGPUAdapterStoreKey);
+	FPlatformMisc::DeleteStoredValue(EditorAnalyticsDefs::StoreId, SectionName, EditorAnalyticsDefs::RenderingGPUAdapterStoreKey);
+	FPlatformMisc::DeleteStoredValue(EditorAnalyticsDefs::StoreId, SectionName, EditorAnalyticsDefs::GPUVendorIDStoreKey);
+	FPlatformMisc::DeleteStoredValue(EditorAnalyticsDefs::StoreId, SectionName, EditorAnalyticsDefs::GPUDeviceIDStoreKey);
+	FPlatformMisc::DeleteStoredValue(EditorAnalyticsDefs::StoreId, SectionName, EditorAnalyticsDefs::GRHIDeviceRevisionStoreKey);
+	FPlatformMisc::DeleteStoredValue(EditorAnalyticsDefs::StoreId, SectionName, EditorAnalyticsDefs::GRHIAdapterInternalDriverVersionStoreKey);
+	FPlatformMisc::DeleteStoredValue(EditorAnalyticsDefs::StoreId, SectionName, EditorAnalyticsDefs::GRHIAdapterUserDriverVersionStoreKey);
+
+	FPlatformMisc::DeleteStoredValue(EditorAnalyticsDefs::StoreId, SectionName, EditorAnalyticsDefs::TotalPhysicalRAMStoreKey);
+	FPlatformMisc::DeleteStoredValue(EditorAnalyticsDefs::StoreId, SectionName, EditorAnalyticsDefs::CPUPhysicalCoresStoreKey);
+	FPlatformMisc::DeleteStoredValue(EditorAnalyticsDefs::StoreId, SectionName, EditorAnalyticsDefs::CPULogicalCoresStoreKey);
+	FPlatformMisc::DeleteStoredValue(EditorAnalyticsDefs::StoreId, SectionName, EditorAnalyticsDefs::CPUVendorStoreKey);
+	FPlatformMisc::DeleteStoredValue(EditorAnalyticsDefs::StoreId, SectionName, EditorAnalyticsDefs::CPUBrandStoreKey);
+
+	FPlatformMisc::DeleteStoredValue(EditorAnalyticsDefs::StoreId, SectionName, EditorAnalyticsDefs::OSMajorStoreKey);
+	FPlatformMisc::DeleteStoredValue(EditorAnalyticsDefs::StoreId, SectionName, EditorAnalyticsDefs::OSMinorStoreKey);
+	FPlatformMisc::DeleteStoredValue(EditorAnalyticsDefs::StoreId, SectionName, EditorAnalyticsDefs::OSVersionStoreKey);
+	FPlatformMisc::DeleteStoredValue(EditorAnalyticsDefs::StoreId, SectionName, EditorAnalyticsDefs::bIs64BitOSStoreKey);
+
+	FPlatformMisc::DeleteStoredValue(EditorAnalyticsDefs::StoreId, SectionName, EditorAnalyticsDefs::IsCrashStoreKey);
+	FPlatformMisc::DeleteStoredValue(EditorAnalyticsDefs::StoreId, SectionName, EditorAnalyticsDefs::IsGPUCrashStoreKey);
+	FPlatformMisc::DeleteStoredValue(EditorAnalyticsDefs::StoreId, SectionName, EditorAnalyticsDefs::IsDebuggerStoreKey);
+	FPlatformMisc::DeleteStoredValue(EditorAnalyticsDefs::StoreId, SectionName, EditorAnalyticsDefs::WasDebuggerStoreKey);
+	FPlatformMisc::DeleteStoredValue(EditorAnalyticsDefs::StoreId, SectionName, EditorAnalyticsDefs::IsVanillaStoreKey);
+	FPlatformMisc::DeleteStoredValue(EditorAnalyticsDefs::StoreId, SectionName, EditorAnalyticsDefs::IsTerminatingKey);
+	FPlatformMisc::DeleteStoredValue(EditorAnalyticsDefs::StoreId, SectionName, EditorAnalyticsDefs::WasShutdownStoreKey);
+	FPlatformMisc::DeleteStoredValue(EditorAnalyticsDefs::StoreId, SectionName, EditorAnalyticsDefs::IsInPIEStoreKey);
+	FPlatformMisc::DeleteStoredValue(EditorAnalyticsDefs::StoreId, SectionName, EditorAnalyticsDefs::IsInEnterpriseStoreKey);
+	FPlatformMisc::DeleteStoredValue(EditorAnalyticsDefs::StoreId, SectionName, EditorAnalyticsDefs::IsInVRModeStoreKey);
+
+	return true;
+}
+
+bool FEditorAnalyticsSession::GetStoredSessionIDs(TArray<FString>& OutSessions)
+{
+	if (!ensure(IsLocked()))
+	{
+		return false;
+	}
+
+	OutSessions = EditorAnalyticsUtils::GetSessionList();
+	return true;
+}
+
+bool FEditorAnalyticsSession::LoadAllStoredSessions(TArray<FEditorAnalyticsSession>& OutSessions)
+{
+	if (!ensure(IsLocked()))
+	{
+		return false;
+	}
+
+	TArray<FString> SessionIDs = EditorAnalyticsUtils::GetSessionList();
+
+	// Retrieve all the sessions in the list from storage
+	for (const FString& Id : SessionIDs)
+	{
+		FEditorAnalyticsSession NewSession;
+		EditorAnalyticsUtils::LoadInternal(NewSession, Id);
+
+		OutSessions.Add(NewSession);
+	}
+
+	return true;
+}
+
+bool FEditorAnalyticsSession::SaveStoredSessionIDs(const TArray<FString>& InSessions)
+{
+	// build up a new SessionList string
+	FString SessionListString;
+	for (const FString& Session : InSessions)
+	{
+		if (!SessionListString.IsEmpty())
+		{
+			SessionListString.Append(TEXT(","));
+		}
+
+		SessionListString.Append(Session);
+	}
+
+	if (!ensure(IsLocked()))
+	{
+		return false;
+	}
+
+	FPlatformMisc::SetStoredValue(EditorAnalyticsDefs::StoreId, EditorAnalyticsDefs::SessionSummarySection, EditorAnalyticsDefs::SessionListStoreKey, SessionListString);
+	return true;
+}
diff --git a/Engine/Source/Developer/EditorAnalyticsSession/Private/EditorSessionSummarySender.cpp b/Engine/Source/Developer/EditorAnalyticsSession/Private/EditorSessionSummarySender.cpp
new file mode 100644
--- /dev/null
+++ b/Engine/Source/Developer/EditorAnalyticsSession/Private/EditorSessionSummarySender.cpp
@@ -0,0 +1,171 @@
+// Copyright 1998-2019 Epic Games, Inc. All Rights Reserved.
+
+#include "EditorSessionSummarySender.h"
+
+#include "AnalyticsEventAttribute.h"
+#include "Algo/Transform.h"
+#include "Interfaces/IAnalyticsProvider.h"
+
+DEFINE_LOG_CATEGORY_STATIC(LogEditorSessionSummary, Verbose, All);
+
+/* FEditorSessionSummarySender */
+
+namespace EditorSessionSenderDefs
+{
+	static const FTimespan SessionExpiration = FTimespan::FromDays(30.0);
+	static const int HeartbeatPeriodSeconds = 60;
+
+	// shutdown types
+	static const FString RunningSessionToken(TEXT("Running"));
+	static const FString ShutdownSessionToken(TEXT("Shutdown"));
+	static const FString CrashSessionToken(TEXT("Crashed"));
+	static const FString TerminatedSessionToken(TEXT("Terminated"));
+	static const FString DebuggerSessionToken(TEXT("Debugger"));
+	static const FString AbnormalSessionToken(TEXT("AbnormalShutdown"));
+}
+
+FEditorSessionSummarySender::FEditorSessionSummarySender(IAnalyticsProvider& InAnalyticsProvider, const FString& InSenderName) :
+	HeartbeatTimeElapsed(0.0f)
+	, AnalyticsProvider(InAnalyticsProvider)
+	, Sender(InSenderName)
+{
+	SendStoredSessions();
+}
+
+void FEditorSessionSummarySender::Tick(float DeltaTime)
+{
+	HeartbeatTimeElapsed += DeltaTime;
+
+	if (HeartbeatTimeElapsed > (float) EditorSessionSenderDefs::HeartbeatPeriodSeconds)
+	{
+		HeartbeatTimeElapsed = 0.0f;
+
+		SendStoredSessions();
+	}
+}
+
+void FEditorSessionSummarySender::Shutdown()
+{
+	SendStoredSessions();
+}
+
+void FEditorSessionSummarySender::SendStoredSessions() const
+{
+	TArray<FEditorAnalyticsSession> SessionsToReport;
+
+	if (FEditorAnalyticsSession::Lock(FTimespan::FromMilliseconds(100)))
+	{
+		// Get list of sessions in storage
+		TArray<FEditorAnalyticsSession> ExistingSessions;
+		FEditorAnalyticsSession::LoadAllStoredSessions(ExistingSessions);
+
+		TArray<FEditorAnalyticsSession> SessionsToDelete;
+
+		// Check each stored session to see if they should be sent or not 
+		for (FEditorAnalyticsSession& Session : ExistingSessions)
+		{
+			if (FPlatformProcess::IsApplicationRunning(Session.PlatformProcessID))
+			{
+				// Skip processes that are still running
+				continue;
+			}
+
+			const FTimespan SessionAge = FDateTime::UtcNow() - Session.Timestamp;
+			if (SessionAge < EditorSessionSenderDefs::SessionExpiration)
+			{
+				SessionsToReport.Add(Session);
+			}
+			SessionsToDelete.Add(Session);
+		}
+
+		for (const FEditorAnalyticsSession& ToDelete : SessionsToDelete)
+		{
+			ToDelete.Delete();
+			ExistingSessions.RemoveAll([&ToDelete](const FEditorAnalyticsSession& Session)
+				{
+					return Session.SessionId == ToDelete.SessionId;
+				});
+		}
+
+		TArray<FString> SessionIDs;
+		SessionIDs.Reserve(ExistingSessions.Num());
+		Algo::Transform(ExistingSessions, SessionIDs, &FEditorAnalyticsSession::SessionId);
+
+		FEditorAnalyticsSession::SaveStoredSessionIDs(SessionIDs);
+
+		FEditorAnalyticsSession::Unlock();
+	}
+
+	for (const FEditorAnalyticsSession& Session : SessionsToReport)
+	{
+		SendSessionSummaryEvent(Session);
+	}
+}
+
+void FEditorSessionSummarySender::SendSessionSummaryEvent(const FEditorAnalyticsSession& Session) const
+{
+	FGuid SessionId;
+	FString SessionIdString = Session.SessionId;
+	if (FGuid::Parse(SessionIdString, SessionId))
+	{
+		// convert session guid to one with braces for sending to analytics
+		SessionIdString = SessionId.ToString(EGuidFormats::DigitsWithHyphensInBraces);
+	}
+
+	FString ShutdownTypeString = Session.bCrashed ? EditorSessionSenderDefs::CrashSessionToken :
+		(Session.bWasEverDebugger ? EditorSessionSenderDefs::DebuggerSessionToken :
+		(Session.bIsTerminating ? EditorSessionSenderDefs::TerminatedSessionToken :
+		(Session.bWasShutdown ? EditorSessionSenderDefs::ShutdownSessionToken : EditorSessionSenderDefs::AbnormalSessionToken)));
+
+	FString PluginsString = FString::Join(Session.Plugins, TEXT(","));
+
+	TArray<FAnalyticsEventAttribute> AnalyticsAttributes;
+	AnalyticsAttributes.Emplace(TEXT("ProjectName"), Session.ProjectName);
+	AnalyticsAttributes.Emplace(TEXT("ProjectID"), Session.ProjectID);
+	AnalyticsAttributes.Emplace(TEXT("ProjectDescription"), Session.ProjectDescription);
+	AnalyticsAttributes.Emplace(TEXT("ProjectVersion"), Session.ProjectVersion);
+	AnalyticsAttributes.Emplace(TEXT("Platform"), FPlatformProperties::PlatformName());
+	AnalyticsAttributes.Emplace(TEXT("SessionId"), SessionIdString);
+	AnalyticsAttributes.Emplace(TEXT("EngineVersion"), Session.EngineVersion);
+	AnalyticsAttributes.Emplace(TEXT("ShutdownType"), ShutdownTypeString);
+	AnalyticsAttributes.Emplace(TEXT("Timestamp"), Session.Timestamp.ToIso8601());
+	AnalyticsAttributes.Emplace(TEXT("CurrentUserActivity"), Session.CurrentUserActivity);
+	AnalyticsAttributes.Emplace(TEXT("StartupTimestamp"), Session.StartupTimestamp.ToIso8601());
+	AnalyticsAttributes.Emplace(TEXT("AverageFPS"), Session.AverageFPS);
+
+	double SessionDuration = (Session.Timestamp - Session.StartupTimestamp).GetTotalSeconds();
+	AnalyticsAttributes.Emplace(TEXT("SessionDuration"), SessionDuration);
+
+	AnalyticsAttributes.Emplace(TEXT("1MinIdle"), Session.Idle1Min);
+	AnalyticsAttributes.Emplace(TEXT("5MinIdle"), Session.Idle5Min);
+	AnalyticsAttributes.Emplace(TEXT("30MinIdle"), Session.Idle30Min);
+	AnalyticsAttributes.Emplace(TEXT("DesktopGPUAdapter"), Session.DesktopGPUAdapter);
+	AnalyticsAttributes.Emplace(TEXT("RenderingGPUAdapter"), Session.RenderingGPUAdapter);
+	AnalyticsAttributes.Emplace(TEXT("GPUVendorID"), Session.GPUVendorID);
+	AnalyticsAttributes.Emplace(TEXT("GPUDeviceID"), Session.GPUDeviceID);
+	AnalyticsAttributes.Emplace(TEXT("GRHIDeviceRevision"), Session.GRHIDeviceRevision);
+	AnalyticsAttributes.Emplace(TEXT("GRHIAdapterInternalDriverVersion"), Session.GRHIAdapterInternalDriverVersion);
+	AnalyticsAttributes.Emplace(TEXT("GRHIAdapterUserDriverVersion"), Session.GRHIAdapterUserDriverVersion);
+	AnalyticsAttributes.Emplace(TEXT("TotalPhysicalRAM"), Session.TotalPhysicalRAM);
+	AnalyticsAttributes.Emplace(TEXT("CPUPhysicalCores"), Session.CPUPhysicalCores);
+	AnalyticsAttributes.Emplace(TEXT("CPULogicalCores"), Session.CPULogicalCores);
+	AnalyticsAttributes.Emplace(TEXT("CPUVendor"), Session.CPUVendor);
+	AnalyticsAttributes.Emplace(TEXT("CPUBrand"), Session.CPUBrand);
+	AnalyticsAttributes.Emplace(TEXT("OSMajor"), Session.OSMajor);
+	AnalyticsAttributes.Emplace(TEXT("OSMinor"), Session.OSMinor);
+	AnalyticsAttributes.Emplace(TEXT("OSVersion"), Session.OSVersion);
+	AnalyticsAttributes.Emplace(TEXT("Is64BitOS"), Session.bIs64BitOS);
+
+	AnalyticsAttributes.Emplace(TEXT("GPUCrash"), Session.bGPUCrashed);
+	AnalyticsAttributes.Emplace(TEXT("WasDebugged"), Session.bWasEverDebugger);
+	AnalyticsAttributes.Emplace(TEXT("IsVanilla"), Session.bIsVanilla);
+	AnalyticsAttributes.Emplace(TEXT("WasShutdown"), Session.bWasShutdown);
+	AnalyticsAttributes.Emplace(TEXT("IsInPIE"), Session.bIsInPIE);
+	AnalyticsAttributes.Emplace(TEXT("IsInEnterprise"), Session.bIsInEnterprise);
+	AnalyticsAttributes.Emplace(TEXT("IsInVRMode"), Session.bIsInVRMode);
+	AnalyticsAttributes.Emplace(TEXT("SentFrom"), Sender);
+
+	AnalyticsProvider.RecordEvent(TEXT("SessionSummary"), AnalyticsAttributes);
+
+	UE_LOG(LogEditorSessionSummary, Log, TEXT("EditorSessionSummary sent report. Type=%s, SessionId=%s"), *ShutdownTypeString, *SessionIdString);
+}
diff --git a/Engine/Source/Developer/EditorAnalyticsSession/Public/EditorAnalyticsSession.h b/Engine/Source/Developer/EditorAnalyticsSession/Public/EditorAnalyticsSession.h
new file mode 100644
--- /dev/null
+++ b/Engine/Source/Developer/EditorAnalyticsSession/Public/EditorAnalyticsSession.h
@@ -0,0 +1,116 @@
+// Copyright 1998-2019 Epic Games, Inc. All Rights Reserved.
+
+#include "Modules/ModuleInterface.h"
+
+struct EDITORANALYTICSSESSION_API FEditorAnalyticsSession
+{
+	FString SessionId;
+
+	FString ProjectName;
+	FString ProjectID;
+	FString ProjectDescription;
+	FString ProjectVersion;
+	FString EngineVersion;
+	int32 PlatformProcessID;
+
+	FDateTime StartupTimestamp;
+	FDateTime Timestamp;
+	int32 Idle1Min;
+	int32 Idle5Min;
+	int32 Idle30Min;
+	FString CurrentUserActivity;
+	TArray<FString> Plugins;
+	float AverageFPS;
+
+	FString DesktopGPUAdapter;
+	FString RenderingGPUAdapter;
+	uint32 GPUVendorID;
+	uint32 GPUDeviceID;
+	uint32 GRHIDeviceRevision;
+	FString GRHIAdapterInternalDriverVersion;
+	FString GRHIAdapterUserDriverVersion;
+
+	uint64 TotalPhysicalRAM;
+	int32 CPUPhysicalCores;
+	int32 CPULogicalCores;
+	FString CPUVendor;
+	FString CPUBrand;
+
+	FString OSMajor;
+	FString OSMinor;
+	FString OSVersion;
+	bool bIs64BitOS : 1;
+
+	bool bCrashed : 1;
+	bool bGPUCrashed : 1;
+	bool bIsDebugger : 1;
+	bool bWasEverDebugger : 1;
+	bool bIsVanilla : 1;
+	bool bIsTerminating : 1;
+	bool bWasShutdown : 1;
+	bool bIsInPIE : 1;
+	bool bIsInEnterprise : 1;
+	bool bIsInVRMode : 1;
+
+	FEditorAnalyticsSession();
+
+	/** 
+	 * Save this session to stored values.
+	 * @returns true if the session was successfully saved.
+	 */
+	bool Save();
+
+	/**
+	 *  Load a session with the given session ID from stored values.
+	 * @returns true if the session was found and successfully loaded.
+	 */
+	bool Load(const FString& InSessionID);
+
+	/**
+	 * Delete the stored values of this session.
+	 * Does not modify the actual session object.
+	 * @returns true if the session was successfully deleted.
+	 */
+	bool Delete() const;
+
+	/**
+	 * Retrieve a list of session IDs that are currently stored locally.
+	 * @returns true if the session IDs were successfully retrieved.
+	 */
+	static bool GetStoredSessionIDs(TArray<FString>& OutSessions);
+
+	/**
+	 * Read all stored sessions into the given array.
+	 * @returns true if the sessions were successfully loaded.
+	 */
+	static bool LoadAllStoredSessions(TArray<FEditorAnalyticsSession>& OutSessions);
+
+	/**
+	 * Save the given session IDs to storage.
+	 * @returns true if the session IDs were successfully saved.
+	 */
+	static bool SaveStoredSessionIDs(const TArray<FString>& InSessions);
+
+	/**
+	 * Acquire a lock for local storage.
+	 * @returns true if the lock was acquired successfully.
+	 */
+	static bool Lock(FTimespan Timeout = FTimespan::Zero());
+
+	/**
+	 * Unlock the local storage.
+	 */
+	static void Unlock();
+
+	/** Is the local storage already locked? */
+	static bool IsLocked();
+
+private:
+
+	static FSystemWideCriticalSection* StoredValuesLock;
+};
+
+class FEditorAnalyticsSessionModule : public IModuleInterface
+{
+
+};
\ No newline at end of file
diff --git a/Engine/Source/Developer/EditorAnalyticsSession/Public/EditorSessionSummarySender.h b/Engine/Source/Developer/EditorAnalyticsSession/Public/EditorSessionSummarySender.h
new file mode 100644
--- /dev/null
+++ b/Engine/Source/Developer/EditorAnalyticsSession/Public/EditorSessionSummarySender.h
@@ -0,0 +1,29 @@
+// Copyright 1998-2019 Epic Games, Inc. All Rights Reserved.
+
+#pragma once
+
+struct FEditorAnalyticsSession;
+class IAnalyticsProvider;
+
+/**
+  * Sender of SessionSummary events from all editor sessions in-between runs.
+  * Separated from Writer to make it easier to run it out-of-process.
+  */
+class EDITORANALYTICSSESSION_API FEditorSessionSummarySender
+{
+public:
+	FEditorSessionSummarySender(IAnalyticsProvider& InAnalyticsProvider, const FString& InSenderName);
+
+	void Tick(float DeltaTime);
+	void Shutdown();
+
+private:
+	/** Send any stored Sessions. */
+	void SendStoredSessions() const;
+	void SendSessionSummaryEvent(const FEditorAnalyticsSession& Session) const;
+
+private:
+	float HeartbeatTimeElapsed;
+	IAnalyticsProvider& AnalyticsProvider;
+	FString Sender;
+};
\ No newline at end of file
diff --git a/Engine/Source/Programs/CrashReportClient/CrashReportClient.Build.cs b/Engine/Source/Programs/CrashReportClient/CrashReportClient.Build.cs
--- a/Engine/Source/Programs/CrashReportClient/CrashReportClient.Build.cs
+++ b/Engine/Source/Programs/CrashReportClient/CrashReportClient.Build.cs
@@ -34,12 +34,7 @@ public class CrashReportClient : ModuleRules
 				"LauncherPlatform",
 				"Concert",
 				"InputCore",
-			}
-		);
-
-		PrivateDependencyModuleNames.AddRange(
-			new string[] 
-			{
+				"EditorAnalyticsSession",
 				"Slate",
 				"SlateCore",
 				"StandaloneRenderer",
diff --git a/Engine/Source/Programs/CrashReportClient/Private/CrashReportClientApp.cpp b/Engine/Source/Programs/CrashReportClient/Private/CrashReportClientApp.cpp
--- a/Engine/Source/Programs/CrashReportClient/Private/CrashReportClientApp.cpp
+++ b/Engine/Source/Programs/CrashReportClient/Private/CrashReportClientApp.cpp
@@ -35,6 +35,7 @@
 #include "PlatformErrorReport.h"
 #include "XmlFile.h"
 #include "RecoveryService.h"
+#include "EditorSessionSummarySender.h"
 
 /** Default main window size */
 const FVector2D InitialWindowDimensions(740, 560);
@@ -74,6 +75,9 @@ static void* MonitorWritePipe = nullptr;
 /** True to enable the disaster recovery service. */
 static bool bDisasterRecoveryServiceEnabled = false;
 
+/** Is the out-of-process editor session summary sender enabled? */
+static bool bEditorSessionSummaryEnabled = true;
+
 /** Result of submission of report */
 enum SubmitCrashReportResult {
 	Failed,				// Failed to send report
@@ -612,6 +616,21 @@ void RunCrashReportClient(const TCHAR* CommandLine)
 			RecoveryService = MakeUnique<FRecoveryService>(MonitorPid);
 		}
 
+		FCrashReportAnalytics::Initialize();
+
+		TUniquePtr<FEditorSessionSummarySender> EditorSessionSummarySender;
+		if (FCrashReportCoreConfig::Get().GetAllowToBeContacted())
+		{
+			EditorSessionSummarySender = MakeUnique<FEditorSessionSummarySender>(FCrashReportAnalytics::GetProvider(), TEXT("CrashReportClient"));
+
+			FTicker::GetCoreTicker().AddTicker(TEXT("EditorSessionSummarySender"), 0, 
+				[&EditorSessionSummarySender](float DeltaTime)
+				{
+					EditorSessionSummarySender->Tick(DeltaTime);
+					return true;
+				});
+		}
+
 		// This IsApplicationAlive() call is quite expensive, perform it at low frequency.
 		bool bApplicationAlive = FPlatformProcess::IsApplicationAlive(MonitorPid);
 		while (bApplicationAlive && !IsEngineExitRequested())
@@ -625,14 +644,6 @@ void RunCrashReportClient(const TCHAR* CommandLine)
 				FSharedCrashContext CrashContext;
 				if (IsCrashReportAvailable(MonitorPid, CrashContext, MonitorReadPipe))
 				{
-					// When running in monitor mode we cannot setup analytics info on the command line (we start earlier than 
-					// config system), instead the crash context contains what the user has selected.
-					bool bReportCrashAnalyticInfo = CrashContext.bSendUsageData;
-					if (bReportCrashAnalyticInfo)
-					{
-						FCrashReportAnalytics::Initialize();
-					}
-
 					// Build error report in memory.
 					FPlatformErrorReport ErrorReport = CollectErrorReport(RecoveryService.Get(), MonitorPid, CrashContext, MonitorWritePipe);
 					const SubmitCrashReportResult Result = SendErrorReport(ErrorReport, CrashContext.bNoDialog && CrashContext.bSendUnattenededBugReports);
@@ -643,6 +654,7 @@ void RunCrashReportClient(const TCHAR* CommandLine)
 					uint8 ResponseCode[] = { 0xd, 0xe, 0xa, 0xd };
 					FPlatformProcess::WritePipe(MonitorWritePipe, ResponseCode, sizeof(ResponseCode));
 
+					bool bReportCrashAnalyticInfo = CrashContext.bSendUsageData;
 					if (bReportCrashAnalyticInfo)
 					{
 						// If analytics is enabled make sure they are submitted now.
@@ -656,11 +668,6 @@ void RunCrashReportClient(const TCHAR* CommandLine)
 					{
 						GIsRequestingExit = false;
 					}*/
-
-					if (bReportCrashAnalyticInfo)
-					{
-						FCrashReportAnalytics::Shutdown();
-					}
 				}
 			}
 
@@ -688,6 +695,13 @@ void RunCrashReportClient(const TCHAR* CommandLine)
 
 			LastTime = CurrentTime;
 		}
+
+		if (EditorSessionSummarySender.IsValid())
+		{
+			EditorSessionSummarySender->Shutdown();
+		}
+
+		FCrashReportAnalytics::Shutdown();
 	}
 
 	FPrimaryCrashProperties::Shutdown();
diff --git a/Engine/Source/Runtime/Engine/Engine.Build.cs b/Engine/Source/Runtime/Engine/Engine.Build.cs
--- a/Engine/Source/Runtime/Engine/Engine.Build.cs
+++ b/Engine/Source/Runtime/Engine/Engine.Build.cs
@@ -39,15 +39,11 @@ public class Engine : ModuleRules
             }
 		);
 
-		if (Target.Configuration != UnrealTargetConfiguration.Shipping)
-		{
-			PrivateIncludePathModuleNames.AddRange(new string[] { "TaskGraph" });
-		}
-
 		if (Target.Configuration != UnrealTargetConfiguration.Shipping)
 		{
 			PrivateIncludePathModuleNames.AddRange(
 				new string[] {
+					"TaskGraph",
 					"SlateReflector",
 				}
 			);
@@ -57,6 +53,13 @@ public class Engine : ModuleRules
 					"SlateReflector",
 				}
 			);
+
+			PrivateDependencyModuleNames.AddRange(
+				new string[]
+				{
+					"EditorAnalyticsSession",
+				}
+			);
 		}
 
 		PublicDependencyModuleNames.AddRange(
@@ -102,8 +105,8 @@ public class Engine : ModuleRules
 				"AudioMixer",
 				"AudioMixerCore",
 				"SignalProcessing",
-                "CrunchCompression"
-            }
+				"CrunchCompression",
+			}
 		);
 
 		if(Target.Platform != UnrealTargetPlatform.HTML5)
diff --git a/Engine/Source/Runtime/Engine/Private/Analytics/EditorSessionSummary.cpp b/Engine/Source/Runtime/Engine/Private/Analytics/EditorSessionSummary.cpp
deleted file mode 100644
diff --git a/Engine/Source/Runtime/Engine/Private/Analytics/EditorSessionSummary.h b/Engine/Source/Runtime/Engine/Private/Analytics/EditorSessionSummary.h
deleted file mode 100644
diff --git a/Engine/Source/Runtime/Engine/Private/Analytics/EditorSessionSummaryWriter.cpp b/Engine/Source/Runtime/Engine/Private/Analytics/EditorSessionSummaryWriter.cpp
new file mode 100644
--- /dev/null
+++ b/Engine/Source/Runtime/Engine/Private/Analytics/EditorSessionSummaryWriter.cpp
@@ -0,0 +1,347 @@
+// Copyright 1998-2019 Epic Games, Inc. All Rights Reserved.
+
+#include "EditorSessionSummaryWriter.h"
+
+#include "EditorAnalyticsSession.h"
+#include "EngineAnalytics.h"
+#include "EngineGlobals.h"
+#include "GeneralProjectSettings.h"
+#include "RHI.h"
+#include "UserActivityTracking.h"
+
+#include "Engine/Engine.h"
+#include "Framework/Application/SlateApplication.h"
+#include "HAL/PlatformProcess.h"
+#include "Interfaces/IAnalyticsProvider.h"
+#include "Interfaces/IPluginManager.h"
+#include "Interfaces/IProjectManager.h"
+#include "Logging/LogMacros.h"
+#include "Misc/App.h"
+#include "Misc/CoreDelegates.h"
+#include "Misc/EngineBuildSettings.h"
+#include "Misc/EngineVersion.h"
+#include "Misc/Guid.h"
+
+#if WITH_EDITOR
+#include "IVREditorModule.h"
+#include "Kismet2/DebuggerCommands.h"
+#endif // WITH_EDITOR
+
+#define LOCTEXT_NAMESPACE "SessionSummary"
+
+DEFINE_LOG_CATEGORY_STATIC(LogEditorSessionSummary, Verbose, All);
+
+/* FEditorSessionSummaryWriter */
+
+namespace EditorSessionWriterDefs
+{
+	static const int HeartbeatPeriodSeconds = 60;
+}
+
+FEditorSessionSummaryWriter::FEditorSessionSummaryWriter() :
+	CurrentSession(nullptr)
+	, LastUserInteractionTime(0.0f)
+	, HeartbeatTimeElapsed(0.0f)
+	, bShutdown(false)
+{
+
+}
+
+void FEditorSessionSummaryWriter::Initialize()
+{
+	// Register for crash and app state callbacks
+	FCoreDelegates::OnHandleSystemError.AddRaw(this, &FEditorSessionSummaryWriter::OnCrashing);
+	FCoreDelegates::ApplicationWillTerminateDelegate.AddRaw(this, &FEditorSessionSummaryWriter::OnTerminate);
+	FUserActivityTracking::OnActivityChanged.AddRaw(this, &FEditorSessionSummaryWriter::OnUserActivity);
+	FCoreDelegates::IsVanillaProductChanged.AddRaw(this, &FEditorSessionSummaryWriter::OnVanillaStateChanged);
+	FSlateApplication::Get().GetOnModalLoopTickEvent().AddRaw(this, &FEditorSessionSummaryWriter::Tick);
+
+	InitializeSessions();
+}
+
+void FEditorSessionSummaryWriter::InitializeSessions()
+{
+	if (!FEngineAnalytics::IsAvailable() || CurrentSession != nullptr)
+	{
+		return;
+	}
+
+	UE_LOG(LogEditorSessionSummary, Verbose, TEXT("Initializing EditorSessionSummaryWriter for editor session tracking"));
+
+	if (FEditorAnalyticsSession::Lock())
+	{
+		// Create a session Session for this session
+		CurrentSession = CreateCurrentSession();
+		CurrentSession->Save();
+
+		UE_LOG(LogEditorSessionSummary, Log, TEXT("EditorSessionSummaryWriter initialized"));
+
+		// update session list string
+		TArray<FString> StoredSessions;
+
+		FEditorAnalyticsSession::GetStoredSessionIDs(StoredSessions);
+		StoredSessions.Add(CurrentSession->SessionId);
+
+		FEditorAnalyticsSession::SaveStoredSessionIDs(StoredSessions);
+
+		FEditorAnalyticsSession::Unlock();
+	}
+}
+
+void FEditorSessionSummaryWriter::UpdateTimestamps()
+{
+	CurrentSession->Timestamp = FDateTime::UtcNow();
+
+	const float IdleSeconds = FPlatformTime::Seconds() - LastUserInteractionTime;
+
+	// 1 + 1 minutes
+	if (IdleSeconds > (60 + 60))
+	{
+		CurrentSession->Idle1Min += 1;
+	}
+
+	// 5 + 1 minutes
+	if (IdleSeconds > (5 * 60 + 60))
+	{
+		CurrentSession->Idle5Min += 1;
+	}
+
+	// 30 + 1 minutes
+	if (IdleSeconds > (30 * 60 + 60))
+	{
+		CurrentSession->Idle30Min += 1;
+	}
+}
+
+void FEditorSessionSummaryWriter::Tick(float DeltaTime)
+{
+	if (bShutdown)
+	{
+		return;
+	}
+
+	// cache the last user interaction time so that during a crash we have access to it
+	LastUserInteractionTime = FSlateApplication::Get().GetLastUserInteractionTime();
+
+	HeartbeatTimeElapsed += DeltaTime;
+
+	if (HeartbeatTimeElapsed > (float) EditorSessionWriterDefs::HeartbeatPeriodSeconds)
+	{
+		HeartbeatTimeElapsed = 0.0f;
+
+		// Try late initialization
+		InitializeSessions();
+
+		if (CurrentSession != nullptr)
+		{
+			// check if the debugger is present
+			bool bIsDebuggerPresent = FPlatformMisc::IsDebuggerPresent();
+			if (CurrentSession->bIsDebugger != bIsDebuggerPresent)
+			{
+				CurrentSession->bIsDebugger = bIsDebuggerPresent;
+
+				if (!CurrentSession->bWasEverDebugger && CurrentSession->bIsDebugger)
+				{
+					CurrentSession->bWasEverDebugger = true;
+				}
+			}
+
+			UpdateTimestamps();
+
+#if WITH_EDITOR
+			extern ENGINE_API float GAverageFPS;
+
+			CurrentSession->AverageFPS = GAverageFPS;
+			CurrentSession->bIsInVRMode = IVREditorModule::Get().IsVREditorModeActive();
+			CurrentSession->bIsInEnterprise = IProjectManager::Get().IsEnterpriseProject();
+			CurrentSession->bIsInPIE = FPlayWorldCommandCallbacks::IsInPIE();
+#endif
+
+			if (FEditorAnalyticsSession::Lock())
+			{
+				CurrentSession->Save();
+
+				FEditorAnalyticsSession::Unlock();
+			}
+		}
+	}
+}
+
+void FEditorSessionSummaryWriter::Shutdown()
+{
+	FCoreDelegates::OnHandleSystemError.RemoveAll(this);
+	FCoreDelegates::ApplicationHasReactivatedDelegate.RemoveAll(this);
+	FCoreDelegates::ApplicationWillDeactivateDelegate.RemoveAll(this);
+	FCoreDelegates::ApplicationWillEnterBackgroundDelegate.RemoveAll(this);
+	FCoreDelegates::ApplicationHasEnteredForegroundDelegate.RemoveAll(this);
+	FCoreDelegates::ApplicationWillTerminateDelegate.RemoveAll(this);
+	FCoreDelegates::IsVanillaProductChanged.RemoveAll(this);
+
+	if (CurrentSession != nullptr)
+	{
+		if (!CurrentSession->bIsTerminating && !CurrentSession->bCrashed)
+		{
+			CurrentSession->bWasShutdown = true;
+		}
+
+		UpdateTimestamps();
+
+		if (FEditorAnalyticsSession::Lock())
+		{
+			CurrentSession->Save();
+
+			FEditorAnalyticsSession::Unlock();
+		}
+
+		delete CurrentSession;
+		CurrentSession = nullptr;
+		bShutdown = true;
+	}
+}
+
+FEditorAnalyticsSession* FEditorSessionSummaryWriter::CreateCurrentSession() const
+{
+	FEditorAnalyticsSession* Session = new FEditorAnalyticsSession();
+
+	FGuid SessionId;
+	if (FGuid::Parse(FEngineAnalytics::GetProvider().GetSessionID(), SessionId))
+	{
+		// convert session GUID to one without braces or other chars that might not be suitable for storage
+		Session->SessionId = SessionId.ToString(EGuidFormats::DigitsWithHyphens);
+	}
+	else
+	{
+		Session->SessionId = FEngineAnalytics::GetProvider().GetSessionID();
+	}
+
+	const UGeneralProjectSettings& ProjectSettings = *GetDefault<UGeneralProjectSettings>();
+
+	Session->PlatformProcessID = FPlatformProcess::GetCurrentProcessId();
+	Session->ProjectName = ProjectSettings.ProjectName;
+	Session->ProjectID = ProjectSettings.ProjectID.ToString(EGuidFormats::DigitsWithHyphens);
+	Session->ProjectDescription = ProjectSettings.Description;
+	Session->ProjectVersion = ProjectSettings.ProjectVersion;
+	Session->EngineVersion = FEngineVersion::Current().ToString(EVersionComponent::Changelist);
+	Session->Timestamp = Session->StartupTimestamp = FDateTime::UtcNow();
+	Session->bIsDebugger = FPlatformMisc::IsDebuggerPresent();
+	Session->bWasEverDebugger = FPlatformMisc::IsDebuggerPresent();
+	Session->CurrentUserActivity = GetUserActivityString();
+	Session->bIsVanilla = GEngine && GEngine->IsVanillaProduct();
+
+	FString OSMajor;
+	FString OSMinor;
+	FPlatformMisc::GetOSVersions(/*out*/ OSMajor, /*out*/ OSMinor);
+	FPlatformMemoryStats Stats = FPlatformMemory::GetStats();
+
+	Session->DesktopGPUAdapter = FPlatformMisc::GetPrimaryGPUBrand();
+	Session->RenderingGPUAdapter = GRHIAdapterName;
+	Session->GPUVendorID = GRHIVendorId;
+	Session->GPUDeviceID = GRHIDeviceId;
+	Session->GRHIDeviceRevision = GRHIDeviceRevision;
+	Session->GRHIAdapterInternalDriverVersion = GRHIAdapterInternalDriverVersion;
+	Session->GRHIAdapterUserDriverVersion = GRHIAdapterUserDriverVersion;
+	Session->TotalPhysicalRAM = Stats.TotalPhysical;
+	Session->CPUPhysicalCores = FPlatformMisc::NumberOfCores();
+	Session->CPULogicalCores = FPlatformMisc::NumberOfCoresIncludingHyperthreads();
+	Session->CPUVendor = FPlatformMisc::GetCPUVendor();
+	Session->CPUBrand = FPlatformMisc::GetCPUBrand();
+	Session->OSMajor = OSMajor;
+	Session->OSMinor = OSMinor;
+	Session->OSVersion = FPlatformMisc::GetOSVersion();
+	Session->bIs64BitOS = FPlatformMisc::Is64bitOperatingSystem();
+
+	TArray<TSharedRef<IPlugin>> Plugins = IPluginManager::Get().GetEnabledPlugins();
+	for (const TSharedRef<IPlugin>& Plugin : Plugins)
+	{
+		Session->Plugins.Add(Plugin->GetName());
+	}
+
+	Session->Plugins.Sort();
+
+	return Session;
+}
+
+extern CORE_API bool GIsGPUCrashed;
+void FEditorSessionSummaryWriter::OnCrashing()
+{
+	if (CurrentSession != nullptr)
+	{
+		UpdateTimestamps();
+
+		CurrentSession->bCrashed = true;
+		CurrentSession->bGPUCrashed = GIsGPUCrashed;
+
+		if (FEditorAnalyticsSession::Lock())
+		{
+			CurrentSession->Save();
+			
+			FEditorAnalyticsSession::Unlock();
+		}
+	}
+}
+
+void FEditorSessionSummaryWriter::OnTerminate()
+{
+	if (CurrentSession != nullptr)
+	{
+		UpdateTimestamps();
+
+		CurrentSession->bIsTerminating = true;
+
+		if (FEditorAnalyticsSession::Lock())
+		{
+			CurrentSession->Save();
+		
+			FEditorAnalyticsSession::Unlock();
+		}
+
+		if (IsEngineExitRequested())
+		{
+			Shutdown();
+		}
+	}
+}
+
+void FEditorSessionSummaryWriter::OnVanillaStateChanged(bool bIsVanilla)
+{
+	if (CurrentSession != nullptr && CurrentSession->bIsVanilla != bIsVanilla)
+	{
+		CurrentSession->bIsVanilla = bIsVanilla;
+
+		if (FEditorAnalyticsSession::Lock())
+		{
+			CurrentSession->Save();
+		
+			FEditorAnalyticsSession::Unlock();
+		}
+	}
+}
+
+void FEditorSessionSummaryWriter::OnUserActivity(const FUserActivity& UserActivity)
+{
+	if (CurrentSession != nullptr)
+	{
+		CurrentSession->CurrentUserActivity = GetUserActivityString();
+
+		if (FEditorAnalyticsSession::Lock())
+		{
+			CurrentSession->Save();
+		
+			FEditorAnalyticsSession::Unlock();
+		}
+	}
+}
+
+FString FEditorSessionSummaryWriter::GetUserActivityString() const
+{
+	const FUserActivity& UserActivity = FUserActivityTracking::GetUserActivity();
+
+	if (UserActivity.ActionName.IsEmpty())
+	{
+		return TEXT("Unknown");
+	}
+
+	return UserActivity.ActionName;
+}
+
+#undef LOCTEXT_NAMESPACE
diff --git a/Engine/Source/Runtime/Engine/Private/Analytics/EditorSessionSummaryWriter.h b/Engine/Source/Runtime/Engine/Private/Analytics/EditorSessionSummaryWriter.h
new file mode 100644
--- /dev/null
+++ b/Engine/Source/Runtime/Engine/Private/Analytics/EditorSessionSummaryWriter.h
@@ -0,0 +1,39 @@
+// Copyright 1998-2019 Epic Games, Inc. All Rights Reserved.
+
+#pragma once
+
+#include "CoreMinimal.h"
+
+struct FUserActivity;
+struct FEditorAnalyticsSession;
+
+/** Writer for SessionSummary events to track all editor sessions. */
+class FEditorSessionSummaryWriter
+{
+public:
+	FEditorSessionSummaryWriter();
+
+	void Initialize();
+	void Tick(float DeltaTime);
+	void Shutdown();
+
+private:
+	void InitializeSessions();
+
+	FEditorAnalyticsSession* CreateCurrentSession() const;
+
+	void OnCrashing();
+	void OnTerminate();
+	void OnUserActivity(const FUserActivity& UserActivity);
+	void OnVanillaStateChanged(bool bIsVanilla);
+
+	FString GetUserActivityString() const;
+	void UpdateTimestamps();
+
+private:
+	FEditorAnalyticsSession* CurrentSession;
+	FString CurrentSessionSectionName;
+	double LastUserInteractionTime;
+	float HeartbeatTimeElapsed;
+	bool bShutdown;
+};
diff --git a/Engine/Source/Runtime/Engine/Private/EngineAnalytics.cpp b/Engine/Source/Runtime/Engine/Private/EngineAnalytics.cpp
--- a/Engine/Source/Runtime/Engine/Private/EngineAnalytics.cpp
+++ b/Engine/Source/Runtime/Engine/Private/EngineAnalytics.cpp
@@ -15,7 +15,10 @@
 #include "EngineSessionManager.h"
 #include "Misc/EngineVersion.h"
 #include "RHI.h"
-#include "Analytics/EditorSessionSummary.h"
+#include "EditorAnalyticsSession.h"
+#include "EditorSessionSummarySender.h"
+#include "Analytics/EditorSessionSummaryWriter.h"
+#include "GenericPlatform/GenericPlatformCrashContext.h"
 
 bool FEngineAnalytics::bIsInitialized;
 TSharedPtr<IAnalyticsProviderET> FEngineAnalytics::Analytics;
@@ -163,7 +166,11 @@ void FEngineAnalytics::Initialize()
 
 		if (!SessionSummarySender.IsValid())
 		{
-			SessionSummarySender = MakeShareable(new FEditorSessionSummarySender());
+			// if we're using out-of-process crash reporting, then we don't need to create a sender in this process.
+			if (!FGenericCrashContext::IsOutOfProcessCrashReporter())
+			{
+				SessionSummarySender = MakeShareable(new FEditorSessionSummarySender(FEngineAnalytics::GetProvider(), TEXT("Editor")));
+			}
 		}
 	}
 }
@@ -187,7 +194,11 @@ void FEngineAnalytics::Shutdown(bool bIsEngineShutdown)
 		SessionSummaryWriter.Reset();
 	}
 
-	SessionSummarySender.Reset();
+	if (SessionSummarySender.IsValid())
+	{
+		SessionSummarySender->Shutdown();
+		SessionSummarySender.Reset();
+	}
 }
 
 void FEngineAnalytics::Tick(float DeltaTime)
