Fixed a whole host of non-unity compile errors in the EditorAnalyticsSession module.

#rnx

diff --git a/Engine/Source/Developer/EditorAnalyticsSession/Private/EditorAnalyticsSession.cpp b/Engine/Source/Developer/EditorAnalyticsSession/Private/EditorAnalyticsSession.cpp
--- a/Engine/Source/Developer/EditorAnalyticsSession/Private/EditorAnalyticsSession.cpp
+++ b/Engine/Source/Developer/EditorAnalyticsSession/Private/EditorAnalyticsSession.cpp
@@ -1,6 +1,7 @@
 // Copyright 1998-2019 Epic Games, Inc. All Rights Reserved.
 
 #include "EditorAnalyticsSession.h"
+#include "Modules/ModuleManager.h"
 
 IMPLEMENT_MODULE(FEditorAnalyticsSessionModule, EditorAnalyticsSession);
 
diff --git a/Engine/Source/Developer/EditorAnalyticsSession/Private/EditorSessionSummarySender.cpp b/Engine/Source/Developer/EditorAnalyticsSession/Private/EditorSessionSummarySender.cpp
--- a/Engine/Source/Developer/EditorAnalyticsSession/Private/EditorSessionSummarySender.cpp
+++ b/Engine/Source/Developer/EditorAnalyticsSession/Private/EditorSessionSummarySender.cpp
@@ -5,6 +5,8 @@
 #include "AnalyticsEventAttribute.h"
 #include "Algo/Transform.h"
 #include "Interfaces/IAnalyticsProvider.h"
+#include "EditorAnalyticsSession.h"
+#include "HAL/PlatformProcess.h"
 
 DEFINE_LOG_CATEGORY_STATIC(LogEditorSessionSummary, Verbose, All);
 
diff --git a/Engine/Source/Developer/EditorAnalyticsSession/Public/EditorAnalyticsSession.h b/Engine/Source/Developer/EditorAnalyticsSession/Public/EditorAnalyticsSession.h
--- a/Engine/Source/Developer/EditorAnalyticsSession/Public/EditorAnalyticsSession.h
+++ b/Engine/Source/Developer/EditorAnalyticsSession/Public/EditorAnalyticsSession.h
@@ -1,5 +1,10 @@
 // Copyright 1998-2019 Epic Games, Inc. All Rights Reserved.
 
+#pragma once
+
+#include "Containers/UnrealString.h"
+#include "HAL/CriticalSection.h"
+#include "Misc/DateTime.h"
 #include "Modules/ModuleInterface.h"
 
 struct EDITORANALYTICSSESSION_API FEditorAnalyticsSession
diff --git a/Engine/Source/Developer/EditorAnalyticsSession/Public/EditorSessionSummarySender.h b/Engine/Source/Developer/EditorAnalyticsSession/Public/EditorSessionSummarySender.h
--- a/Engine/Source/Developer/EditorAnalyticsSession/Public/EditorSessionSummarySender.h
+++ b/Engine/Source/Developer/EditorAnalyticsSession/Public/EditorSessionSummarySender.h
@@ -2,6 +2,9 @@
 
 #pragma once
 
+#include "Containers/UnrealString.h"
+#include "HAL/Platform.h"
+
 struct FEditorAnalyticsSession;
 class IAnalyticsProvider;
 
