Fix missing includes
#rnx

diff --git a/Engine/Source/Developer/EditorAnalyticsSession/Public/EditorSessionSummarySender.h b/Engine/Source/Developer/EditorAnalyticsSession/Public/EditorSessionSummarySender.h
--- a/Engine/Source/Developer/EditorAnalyticsSession/Public/EditorSessionSummarySender.h
+++ b/Engine/Source/Developer/EditorAnalyticsSession/Public/EditorSessionSummarySender.h
@@ -4,6 +4,7 @@
 
 #include "Containers/UnrealString.h"
 #include "HAL/Platform.h"
+#include "Misc/Optional.h"
 
 struct FEditorAnalyticsSession;
 class IAnalyticsProvider;
diff --git a/Engine/Source/Runtime/Core/Public/Math/IntPoint.h b/Engine/Source/Runtime/Core/Public/Math/IntPoint.h
--- a/Engine/Source/Runtime/Core/Public/Math/IntPoint.h
+++ b/Engine/Source/Runtime/Core/Public/Math/IntPoint.h
@@ -6,6 +6,7 @@
 #include "Misc/AssertionMacros.h"
 #include "Math/UnrealMathUtility.h"
 #include "Containers/UnrealString.h"
+#include "Serialization/StructuredArchive.h"
 #include "Templates/TypeHash.h"
 
 /**
