Batch writing ini files for Unix

diff --git a/Engine/Source/Developer/EditorAnalyticsSession/Private/EditorAnalyticsSession.cpp b/Engine/Source/Developer/EditorAnalyticsSession/Private/EditorAnalyticsSession.cpp
--- a/Engine/Source/Developer/EditorAnalyticsSession/Private/EditorAnalyticsSession.cpp
+++ b/Engine/Source/Developer/EditorAnalyticsSession/Private/EditorAnalyticsSession.cpp
@@ -308,10 +308,6 @@ bool FEditorAnalyticsSession::IsLocked()
 	return StoredValuesLock != nullptr && StoredValuesLock->IsValid();
 }
 
-// Utility macros to make it easier to see that all values are correctly set
-#define SET_STORED_STRING(FieldName) FPlatformMisc::SetStoredValue(EditorAnalyticsDefs::StoreId, StorageLocation, EditorAnalyticsDefs::FieldName ## StoreKey, FieldName)
-#define SET_STORED_INT(FieldName) FPlatformMisc::SetStoredValue(EditorAnalyticsDefs::StoreId, StorageLocation, EditorAnalyticsDefs::FieldName ## StoreKey, FString::FromInt(FieldName))
-
 bool FEditorAnalyticsSession::Save()
 {
 	if (!ensure(IsLocked()))
@@ -323,66 +319,66 @@ bool FEditorAnalyticsSession::Save()
 
 	if (!bAlreadySaved)
 	{
-		SET_STORED_STRING(EngineVersion);
-		SET_STORED_INT(PlatformProcessID);
-		SET_STORED_STRING(DesktopGPUAdapter);
-		SET_STORED_STRING(RenderingGPUAdapter);
-		SET_STORED_INT(GPUVendorID);
-		SET_STORED_INT(GPUDeviceID);
-		SET_STORED_INT(GRHIDeviceRevision);
-		SET_STORED_STRING(GRHIAdapterInternalDriverVersion);
-		SET_STORED_STRING(GRHIAdapterUserDriverVersion);
-
-		FPlatformMisc::SetStoredValue(EditorAnalyticsDefs::StoreId, StorageLocation, EditorAnalyticsDefs::TotalPhysicalRAMStoreKey, FString::Printf(TEXT("%llu"), TotalPhysicalRAM));
-
-		SET_STORED_INT(CPUPhysicalCores);
-		SET_STORED_INT(CPULogicalCores);
-		SET_STORED_STRING(CPUVendor);
-		SET_STORED_STRING(CPUBrand);
-
-		FPlatformMisc::SetStoredValue(EditorAnalyticsDefs::StoreId, StorageLocation, EditorAnalyticsDefs::StartupTimestampStoreKey, EditorAnalyticsUtils::TimestampToString(StartupTimestamp));
-
-		SET_STORED_STRING(OSMajor);
-		SET_STORED_STRING(OSMinor);
-		SET_STORED_STRING(OSVersion);
-
-		FPlatformMisc::SetStoredValue(EditorAnalyticsDefs::StoreId, StorageLocation, EditorAnalyticsDefs::bIs64BitOSStoreKey, EditorAnalyticsUtils::BoolToStoredString(bIs64BitOS));
-
 		const FString PluginsString = FString::Join(Plugins, TEXT(","));
-		FPlatformMisc::SetStoredValue(EditorAnalyticsDefs::StoreId, StorageLocation, EditorAnalyticsDefs::PluginsStoreKey, PluginsString);
 
-		SET_STORED_STRING(AppId);
-		SET_STORED_STRING(AppVersion);
-		SET_STORED_STRING(UserId);
-	
+		TMap<FString, FString> KeyValues = {
+			{EditorAnalyticsDefs::EngineVersionStoreKey,       EngineVersion},
+			{EditorAnalyticsDefs::PlatformProcessIDStoreKey,   FString::FromInt(PlatformProcessID)},
+			{EditorAnalyticsDefs::DesktopGPUAdapterStoreKey,   DesktopGPUAdapter},
+			{EditorAnalyticsDefs::RenderingGPUAdapterStoreKey, RenderingGPUAdapter},
+			{EditorAnalyticsDefs::GPUVendorIDStoreKey,         FString::FromInt(GPUVendorID)},
+			{EditorAnalyticsDefs::GPUDeviceIDStoreKey,         FString::FromInt(GPUDeviceID)},
+			{EditorAnalyticsDefs::GRHIDeviceRevisionStoreKey,  FString::FromInt(GRHIDeviceRevision)},
+			{EditorAnalyticsDefs::GRHIAdapterInternalDriverVersionStoreKey, GRHIAdapterUserDriverVersion},
+			{EditorAnalyticsDefs::GRHIAdapterUserDriverVersionStoreKey,     GRHIAdapterUserDriverVersion},
+			{EditorAnalyticsDefs::TotalPhysicalRAMStoreKey, FString::Printf(TEXT("%llu"), TotalPhysicalRAM)},
+			{EditorAnalyticsDefs::CPUPhysicalCoresStoreKey, FString::FromInt(CPUPhysicalCores)},
+			{EditorAnalyticsDefs::CPULogicalCoresStoreKey,  FString::FromInt(CPULogicalCores)},
+			{EditorAnalyticsDefs::CPUVendorStoreKey,        CPUVendor},
+			{EditorAnalyticsDefs::CPUBrandStoreKey,         CPUBrand},
+			{EditorAnalyticsDefs::StartupTimestampStoreKey, EditorAnalyticsUtils::TimestampToString(StartupTimestamp)},
+			{EditorAnalyticsDefs::OSMajorStoreKey,    OSMajor},
+			{EditorAnalyticsDefs::OSMinorStoreKey,    OSMinor},
+			{EditorAnalyticsDefs::OSVersionStoreKey,  OSVersion},
+			{EditorAnalyticsDefs::bIs64BitOSStoreKey, EditorAnalyticsUtils::BoolToStoredString(bIs64BitOS)},
+			{EditorAnalyticsDefs::PluginsStoreKey,    PluginsString},
+			{EditorAnalyticsDefs::AppIdStoreKey,      AppId},
+			{EditorAnalyticsDefs::AppVersionStoreKey, AppVersion},
+			{EditorAnalyticsDefs::UserIdStoreKey,     UserId},
+		};
+
+		FPlatformMisc::SetStoredValues(EditorAnalyticsDefs::StoreId, StorageLocation, KeyValues);
+
 		bAlreadySaved = true;
 	}
 
-	SET_STORED_STRING(ProjectName);
-	SET_STORED_STRING(ProjectID);
-	SET_STORED_STRING(ProjectDescription);
-	SET_STORED_STRING(ProjectVersion);
-
-	FPlatformMisc::SetStoredValue(EditorAnalyticsDefs::StoreId, StorageLocation, EditorAnalyticsDefs::TimestampStoreKey, EditorAnalyticsUtils::TimestampToString(Timestamp));
-
-	SET_STORED_INT(SessionDuration);
-	SET_STORED_INT(Idle1Min);
-	SET_STORED_INT(Idle5Min);
-	SET_STORED_INT(Idle30Min);
-
-	SET_STORED_STRING(CurrentUserActivity);
-
-	FPlatformMisc::SetStoredValue(EditorAnalyticsDefs::StoreId, StorageLocation, EditorAnalyticsDefs::AverageFPSStoreKey, FString::SanitizeFloat(AverageFPS));
-	FPlatformMisc::SetStoredValue(EditorAnalyticsDefs::StoreId, StorageLocation, EditorAnalyticsDefs::IsCrashStoreKey, EditorAnalyticsUtils::BoolToStoredString(bCrashed));
-	FPlatformMisc::SetStoredValue(EditorAnalyticsDefs::StoreId, StorageLocation, EditorAnalyticsDefs::IsGPUCrashStoreKey, EditorAnalyticsUtils::BoolToStoredString(bGPUCrashed));
-	FPlatformMisc::SetStoredValue(EditorAnalyticsDefs::StoreId, StorageLocation, EditorAnalyticsDefs::IsDebuggerStoreKey, EditorAnalyticsUtils::BoolToStoredString(bIsDebugger));
-	FPlatformMisc::SetStoredValue(EditorAnalyticsDefs::StoreId, StorageLocation, EditorAnalyticsDefs::WasDebuggerStoreKey, EditorAnalyticsUtils::BoolToStoredString(bWasEverDebugger));
-	FPlatformMisc::SetStoredValue(EditorAnalyticsDefs::StoreId, StorageLocation, EditorAnalyticsDefs::IsVanillaStoreKey, EditorAnalyticsUtils::BoolToStoredString(bIsVanilla));
-	FPlatformMisc::SetStoredValue(EditorAnalyticsDefs::StoreId, StorageLocation, EditorAnalyticsDefs::IsTerminatingKey, EditorAnalyticsUtils::BoolToStoredString(bIsTerminating));
-	FPlatformMisc::SetStoredValue(EditorAnalyticsDefs::StoreId, StorageLocation, EditorAnalyticsDefs::WasShutdownStoreKey, EditorAnalyticsUtils::BoolToStoredString(bWasShutdown));
-	FPlatformMisc::SetStoredValue(EditorAnalyticsDefs::StoreId, StorageLocation, EditorAnalyticsDefs::IsInPIEStoreKey, EditorAnalyticsUtils::BoolToStoredString(bIsInPIE));
-	FPlatformMisc::SetStoredValue(EditorAnalyticsDefs::StoreId, StorageLocation, EditorAnalyticsDefs::IsInEnterpriseStoreKey, EditorAnalyticsUtils::BoolToStoredString(bIsInEnterprise));
-	FPlatformMisc::SetStoredValue(EditorAnalyticsDefs::StoreId, StorageLocation, EditorAnalyticsDefs::IsInVRModeStoreKey, EditorAnalyticsUtils::BoolToStoredString(bIsInVRMode));
+	{
+		TMap<FString, FString> KeyValues = {
+			{EditorAnalyticsDefs::ProjectNameStoreKey, ProjectName},
+			{EditorAnalyticsDefs::ProjectIDStoreKey,   ProjectID},
+			{EditorAnalyticsDefs::ProjectDescriptionStoreKey,  ProjectDescription},
+			{EditorAnalyticsDefs::ProjectVersionStoreKey,  ProjectVersion},
+			{EditorAnalyticsDefs::TimestampStoreKey,       EditorAnalyticsUtils::TimestampToString(Timestamp)},
+			{EditorAnalyticsDefs::SessionDurationStoreKey, FString::FromInt(SessionDuration)},
+			{EditorAnalyticsDefs::Idle1MinStoreKey,  FString::FromInt(Idle1Min)},
+			{EditorAnalyticsDefs::Idle5MinStoreKey,  FString::FromInt(Idle5Min)},
+			{EditorAnalyticsDefs::Idle30MinStoreKey, FString::FromInt(Idle30Min)},
+			{EditorAnalyticsDefs::CurrentUserActivityStoreKey, CurrentUserActivity},
+			{EditorAnalyticsDefs::AverageFPSStoreKey,     FString::SanitizeFloat(AverageFPS)},
+			{EditorAnalyticsDefs::IsCrashStoreKey,        EditorAnalyticsUtils::BoolToStoredString(bCrashed)},
+			{EditorAnalyticsDefs::IsGPUCrashStoreKey,     EditorAnalyticsUtils::BoolToStoredString(bGPUCrashed)},
+			{EditorAnalyticsDefs::IsDebuggerStoreKey,     EditorAnalyticsUtils::BoolToStoredString(bIsDebugger)},
+			{EditorAnalyticsDefs::WasDebuggerStoreKey,    EditorAnalyticsUtils::BoolToStoredString(bWasEverDebugger)},
+			{EditorAnalyticsDefs::IsVanillaStoreKey,      EditorAnalyticsUtils::BoolToStoredString(bIsVanilla)},
+			{EditorAnalyticsDefs::IsTerminatingKey,       EditorAnalyticsUtils::BoolToStoredString(bIsTerminating)},
+			{EditorAnalyticsDefs::WasShutdownStoreKey,    EditorAnalyticsUtils::BoolToStoredString(bWasShutdown)},
+			{EditorAnalyticsDefs::IsInPIEStoreKey,        EditorAnalyticsUtils::BoolToStoredString(bIsInPIE)    },
+			{EditorAnalyticsDefs::IsInEnterpriseStoreKey, EditorAnalyticsUtils::BoolToStoredString(bIsInEnterprise)},
+			{EditorAnalyticsDefs::IsInVRModeStoreKey,     EditorAnalyticsUtils::BoolToStoredString(bIsInVRMode)},
+		};
+
+		FPlatformMisc::SetStoredValues(EditorAnalyticsDefs::StoreId, StorageLocation, KeyValues);
+	}
 
 	return true;
 }
@@ -394,25 +390,22 @@ bool FEditorAnalyticsSession::SaveForCrash()
 		return false;
 	}
 
-	// These ini writes are causing MallocCrash to go over its LARGE_MEMORYPOOL_SIZE due to writing to ini files causing a potentially a few large allocations per write
-	if (!PLATFORM_UNIX)
-	{
-		const FString StorageLocation = EditorAnalyticsUtils::GetSessionStorageLocation(SessionId);
-		
-		FPlatformMisc::SetStoredValue(EditorAnalyticsDefs::StoreId, StorageLocation, EditorAnalyticsDefs::IsCrashStoreKey, EditorAnalyticsUtils::BoolToStoredString(bCrashed));
-		FPlatformMisc::SetStoredValue(EditorAnalyticsDefs::StoreId, StorageLocation, EditorAnalyticsDefs::IsGPUCrashStoreKey, EditorAnalyticsUtils::BoolToStoredString(bGPUCrashed));
-		FPlatformMisc::SetStoredValue(EditorAnalyticsDefs::StoreId, StorageLocation, EditorAnalyticsDefs::IsTerminatingKey, EditorAnalyticsUtils::BoolToStoredString(bIsTerminating));
-		FPlatformMisc::SetStoredValue(EditorAnalyticsDefs::StoreId, StorageLocation, EditorAnalyticsDefs::WasShutdownStoreKey, EditorAnalyticsUtils::BoolToStoredString(bWasShutdown));
-		FPlatformMisc::SetStoredValue(EditorAnalyticsDefs::StoreId, StorageLocation, EditorAnalyticsDefs::TimestampStoreKey, EditorAnalyticsUtils::TimestampToString(Timestamp));
-		FPlatformMisc::SetStoredValue(EditorAnalyticsDefs::StoreId, StorageLocation, EditorAnalyticsDefs::SessionDurationStoreKey, FString::FromInt(SessionDuration));
-	}
+	const FString StorageLocation = EditorAnalyticsUtils::GetSessionStorageLocation(SessionId);
+
+	TMap<FString, FString> KeyValues = {
+		{EditorAnalyticsDefs::IsCrashStoreKey,         EditorAnalyticsUtils::BoolToStoredString(bCrashed)},
+		{EditorAnalyticsDefs::IsGPUCrashStoreKey,      EditorAnalyticsUtils::BoolToStoredString(bGPUCrashed)},
+		{EditorAnalyticsDefs::IsTerminatingKey,        EditorAnalyticsUtils::BoolToStoredString(bIsTerminating)},
+		{EditorAnalyticsDefs::WasShutdownStoreKey,     EditorAnalyticsUtils::BoolToStoredString(bWasShutdown)},
+		{EditorAnalyticsDefs::TimestampStoreKey,       EditorAnalyticsUtils::TimestampToString(Timestamp)},
+		{EditorAnalyticsDefs::SessionDurationStoreKey, FString::FromInt(SessionDuration)},
+	};
+
+	FPlatformMisc::SetStoredValues(EditorAnalyticsDefs::StoreId, StorageLocation, KeyValues);
 
 	return true;
 }
 
-#undef SET_STORED_INT
-#undef SET_STORED_STRING
-
 bool FEditorAnalyticsSession::Load(const FString& InSessionID)
 {
 	if (!ensure(IsLocked()))
diff --git a/Engine/Source/Runtime/Core/Private/GenericPlatform/GenericPlatformMisc.cpp b/Engine/Source/Runtime/Core/Private/GenericPlatform/GenericPlatformMisc.cpp
--- a/Engine/Source/Runtime/Core/Private/GenericPlatform/GenericPlatformMisc.cpp
+++ b/Engine/Source/Runtime/Core/Private/GenericPlatform/GenericPlatformMisc.cpp
@@ -529,6 +529,19 @@ void FGenericPlatformMisc::EndNamedEvent()
 #endif
 }
 
+bool FGenericPlatformMisc::SetStoredValues(const FString& InStoreId, const FString& InSectionName, const TMap<FString, FString>& InKeyValues)
+{
+	for (const TPair<FString, FString>& InKeyValue : InKeyValues)
+	{
+		if (!FPlatformMisc::SetStoredValue(InStoreId, InSectionName, InKeyValue.Key, InKeyValue.Value))
+		{
+			return false;
+		}
+	}
+
+	return true;
+}
+
 bool FGenericPlatformMisc::SetStoredValue(const FString& InStoreId, const FString& InSectionName, const FString& InKeyName, const FString& InValue)
 {
 	check(!InStoreId.IsEmpty());
diff --git a/Engine/Source/Runtime/Core/Private/Unix/UnixPlatformMisc.cpp b/Engine/Source/Runtime/Core/Private/Unix/UnixPlatformMisc.cpp
--- a/Engine/Source/Runtime/Core/Private/Unix/UnixPlatformMisc.cpp
+++ b/Engine/Source/Runtime/Core/Private/Unix/UnixPlatformMisc.cpp
@@ -1105,3 +1105,27 @@ IPlatformChunkInstall* FUnixPlatformMisc::GetPlatformChunkInstall()
 
 	return ChunkInstall;
 }
+
+bool FUnixPlatformMisc::SetStoredValues(const FString& InStoreId, const FString& InSectionName, const TMap<FString, FString>& InKeyValues)
+{
+	check(!InStoreId.IsEmpty());
+	check(!InSectionName.IsEmpty());
+
+	const FString ConfigPath = FString(FPlatformProcess::ApplicationSettingsDir()) / InStoreId / FString(TEXT("KeyValueStore.ini"));
+
+	FConfigFile ConfigFile;
+	ConfigFile.Read(ConfigPath);
+
+	for (auto const& InKeyValue : InKeyValues)
+	{
+		check(!InKeyValue.Key.IsEmpty());
+
+		FConfigSection& Section = ConfigFile.FindOrAdd(InSectionName);
+
+		FConfigValue& KeyValue = Section.FindOrAdd(*InKeyValue.Key);
+		KeyValue = FConfigValue(InKeyValue.Value);
+	}
+
+	ConfigFile.Dirty = true;
+	return ConfigFile.Write(ConfigPath);
+}
diff --git a/Engine/Source/Runtime/Core/Public/GenericPlatform/GenericPlatformMisc.h b/Engine/Source/Runtime/Core/Public/GenericPlatform/GenericPlatformMisc.h
--- a/Engine/Source/Runtime/Core/Public/GenericPlatform/GenericPlatformMisc.h
+++ b/Engine/Source/Runtime/Core/Public/GenericPlatform/GenericPlatformMisc.h
@@ -662,6 +662,17 @@ public:
 	{
 	}
 
+	/**
+	 *	Set the value for the given section and key in the platform specific key->value store
+	 *  Note: The key->value store is user-specific, but may be used to share data between different applications for the same user
+	 *
+	 *  @param	InStoreId			The name used to identify the store you want to use (eg, MyGame)
+	 *	@param	InSectionName		The section that this key->value pair is placed within (can contain / separators, eg UserDetails/AccountInfo)
+	 *	@param	InKeyValues			The mapping of key->value pairs to set
+	 *	@return	bool				true if the value was set correctly, false if not
+	 */
+	static bool SetStoredValues(const FString& InStoreId, const FString& InSectionName, const TMap<FString, FString>& InKeyValues);
+
 	/**
 	 *	Set the value for the given section and key in the platform specific key->value store
 	 *  Note: The key->value store is user-specific, but may be used to share data between different applications for the same user
diff --git a/Engine/Source/Runtime/Core/Public/Unix/UnixPlatformMisc.h b/Engine/Source/Runtime/Core/Public/Unix/UnixPlatformMisc.h
--- a/Engine/Source/Runtime/Core/Public/Unix/UnixPlatformMisc.h
+++ b/Engine/Source/Runtime/Core/Public/Unix/UnixPlatformMisc.h
@@ -138,6 +138,8 @@ struct CORE_API FUnixPlatformMisc : public FGenericPlatformMisc
 
 	static IPlatformChunkInstall* GetPlatformChunkInstall();
 
+	static bool SetStoredValues(const FString& InStoreId, const FString& InSectionName, const TMap<FString, FString>& InKeyValues);
+
 #if STATS || ENABLE_STATNAMEDEVENTS
 	static void BeginNamedEventFrame();
 	static void BeginNamedEvent(const struct FColor& Color, const TCHAR* Text);
