- Fixed the EditorSessionSummarySender to use the FPlatformProperties::IniPlatformName() rather than FPlatformProperties::PlatformName(). Assuming CRC runs on the same platform as the Editor, the platform reported will be "Windows", "Mac" and "Linux".

diff --git a/Engine/Source/Developer/EditorAnalyticsSession/Private/EditorSessionSummarySender.cpp b/Engine/Source/Developer/EditorAnalyticsSession/Private/EditorSessionSummarySender.cpp
--- a/Engine/Source/Developer/EditorAnalyticsSession/Private/EditorSessionSummarySender.cpp
+++ b/Engine/Source/Developer/EditorAnalyticsSession/Private/EditorSessionSummarySender.cpp
@@ -108,7 +108,7 @@ void FEditorSessionSummarySender::SendStoredSessions(const bool bForceSendCurren
 
 		TArray<FEditorAnalyticsSession> SessionsToDelete;
 
-		// Check each stored session to see if they should be sent or not 
+		// Check each stored session to see if they should be sent or not
 		for (FEditorAnalyticsSession& Session : ExistingSessions)
 		{
 			const bool bForceSendSession = bForceSendCurrentSession && (Session.PlatformProcessID == CurrentSessionProcessId);
@@ -180,7 +180,7 @@ void FEditorSessionSummarySender::SendSessionSummaryEvent(const FEditorAnalytics
 	AnalyticsAttributes.Emplace(TEXT("ProjectID"), Session.ProjectID);
 	AnalyticsAttributes.Emplace(TEXT("ProjectDescription"), Session.ProjectDescription);
 	AnalyticsAttributes.Emplace(TEXT("ProjectVersion"), Session.ProjectVersion);
-	AnalyticsAttributes.Emplace(TEXT("Platform"), FPlatformProperties::PlatformName());
+	AnalyticsAttributes.Emplace(TEXT("Platform"), FPlatformProperties::IniPlatformName());
 	AnalyticsAttributes.Emplace(TEXT("SessionId"), SessionIdString); // The provider is expected to add it as "SessionID" param in the HTTP request, but keep it for completness, because the formats are slightly different.
 	AnalyticsAttributes.Emplace(TEXT("EngineVersion"), Session.EngineVersion); // The provider is expected to add it as "AppVersion" param in the HTTP request, but keep it for completness, because the formats are slightly different.
 	AnalyticsAttributes.Emplace(TEXT("ShutdownType"), ShutdownTypeString);
