- Used the correct analytic key to read the stored value before sending.

diff --git a/Engine/Source/Developer/EditorAnalyticsSession/Private/EditorAnalyticsSession.cpp b/Engine/Source/Developer/EditorAnalyticsSession/Private/EditorAnalyticsSession.cpp
--- a/Engine/Source/Developer/EditorAnalyticsSession/Private/EditorAnalyticsSession.cpp
+++ b/Engine/Source/Developer/EditorAnalyticsSession/Private/EditorAnalyticsSession.cpp
@@ -190,7 +190,7 @@ namespace EditorAnalyticsUtils
 
 		{
 			FString TotalPhysicalRAMString;
-			FPlatformMisc::GetStoredValue(EditorAnalyticsDefs::StoreId, SectionName, EditorAnalyticsDefs::GRHIAdapterUserDriverVersionStoreKey, TotalPhysicalRAMString);
+			FPlatformMisc::GetStoredValue(EditorAnalyticsDefs::StoreId, SectionName, EditorAnalyticsDefs::TotalPhysicalRAMStoreKey, TotalPhysicalRAMString);
 			Session.TotalPhysicalRAM = FCString::Atoi64(*TotalPhysicalRAMString);
 		}
 
