Fixed missing header files reported by CIS in EditorAnalyticSession.cpp

diff --git a/Engine/Source/Developer/EditorAnalyticsSession/Private/EditorAnalyticsSession.cpp b/Engine/Source/Developer/EditorAnalyticsSession/Private/EditorAnalyticsSession.cpp
--- a/Engine/Source/Developer/EditorAnalyticsSession/Private/EditorAnalyticsSession.cpp
+++ b/Engine/Source/Developer/EditorAnalyticsSession/Private/EditorAnalyticsSession.cpp
@@ -3,6 +3,9 @@
 #include "EditorAnalyticsSession.h"
 #include "Modules/ModuleManager.h"
 #include "Internationalization/Regex.h"
+#include "HAL/PlatformProcess.h"
+#include "HAL/FileManager.h"
+#include "Misc/Paths.h"
 
 IMPLEMENT_MODULE(FEditorAnalyticsSessionModule, EditorAnalyticsSession);
 
