Fixed crash in EditorSessionSummaryWriter if the current session is not created.
  - Added null check in case the session creation is delayed because the system wide lock could not be taken at init time.

Fixed the likely cause of unknown exit code reported by CrashReportClientEditor
  - Fixed EditorSessionSummarySenderto loop/wait until system wide lock could be acquired before returning if bForceSendCurrentSession is true.

diff --git a/Engine/Source/Developer/EditorAnalyticsSession/Private/EditorSessionSummarySender.cpp b/Engine/Source/Developer/EditorAnalyticsSession/Private/EditorSessionSummarySender.cpp
--- a/Engine/Source/Developer/EditorAnalyticsSession/Private/EditorSessionSummarySender.cpp
+++ b/Engine/Source/Developer/EditorAnalyticsSession/Private/EditorSessionSummarySender.cpp
@@ -98,51 +98,59 @@ bool FEditorSessionSummarySender::FindCurrentSession(FEditorAnalyticsSession& Ou
 
 void FEditorSessionSummarySender::SendStoredSessions(const bool bForceSendCurrentSession) const
 {
+	// Load the list of sessions to process. Expect contention on the analytic session lock between the Editor and CrashReportClientEditor (on windows) or between Editor instances (on mac/linux)
+	//   - Try every 'n' seconds if bForceSendCurrentSession is true.
+	//   - Don't block and don't loop if bForceSendCurrentSession is false.
 	TArray<FEditorAnalyticsSession> SessionsToReport;
-
-	if (FEditorAnalyticsSession::Lock(FTimespan::FromMilliseconds(100)))
+	FTimespan Timemout(bForceSendCurrentSession ? FTimespan::FromSeconds(0.5) : FTimespan::Zero());
+	bool bSessionsLoaded = false;
+	do
 	{
-		// Get list of sessions in storage
-		TArray<FEditorAnalyticsSession> ExistingSessions;
-		FEditorAnalyticsSession::LoadAllStoredSessions(ExistingSessions);
+		if (FEditorAnalyticsSession::Lock(Timemout))
+		{
+			// Get list of sessions in storage
+			TArray<FEditorAnalyticsSession> ExistingSessions;
+			FEditorAnalyticsSession::LoadAllStoredSessions(ExistingSessions);
 
-		TArray<FEditorAnalyticsSession> SessionsToDelete;
+			TArray<FEditorAnalyticsSession> SessionsToDelete;
 
-		// Check each stored session to see if they should be sent or not
-		for (FEditorAnalyticsSession& Session : ExistingSessions)
-		{
-			const bool bForceSendSession = bForceSendCurrentSession && (Session.PlatformProcessID == CurrentSessionProcessId);
-			if (!bForceSendSession && FPlatformProcess::IsApplicationRunning(Session.PlatformProcessID))
+			// Check each stored session to see if they should be sent or not
+			for (FEditorAnalyticsSession& Session : ExistingSessions)
 			{
-				// Skip processes that are still running
-				continue;
+				const bool bForceSendSession = bForceSendCurrentSession && (Session.PlatformProcessID == CurrentSessionProcessId);
+				if (!bForceSendSession && FPlatformProcess::IsApplicationRunning(Session.PlatformProcessID))
+				{
+					// Skip processes that are still running
+					continue;
+				}
+
+				const FTimespan SessionAge = FDateTime::UtcNow() - Session.Timestamp;
+				if (SessionAge < EditorSessionSenderDefs::SessionExpiration)
+				{
+					SessionsToReport.Add(Session);
+				}
+				SessionsToDelete.Add(Session);
 			}
 
-			const FTimespan SessionAge = FDateTime::UtcNow() - Session.Timestamp;
-			if (SessionAge < EditorSessionSenderDefs::SessionExpiration)
+			for (const FEditorAnalyticsSession& ToDelete : SessionsToDelete)
 			{
-				SessionsToReport.Add(Session);
+				ToDelete.Delete();
+				ExistingSessions.RemoveAll([&ToDelete](const FEditorAnalyticsSession& Session)
+					{
+						return Session.SessionId == ToDelete.SessionId;
+					});
 			}
-			SessionsToDelete.Add(Session);
-		}
-
-		for (const FEditorAnalyticsSession& ToDelete : SessionsToDelete)
-		{
-			ToDelete.Delete();
-			ExistingSessions.RemoveAll([&ToDelete](const FEditorAnalyticsSession& Session)
-				{
-					return Session.SessionId == ToDelete.SessionId;
-				});
-		}
 
-		TArray<FString> SessionIDs;
-		SessionIDs.Reserve(ExistingSessions.Num());
-		Algo::Transform(ExistingSessions, SessionIDs, &FEditorAnalyticsSession::SessionId);
+			TArray<FString> SessionIDs;
+			SessionIDs.Reserve(ExistingSessions.Num());
+			Algo::Transform(ExistingSessions, SessionIDs, &FEditorAnalyticsSession::SessionId);
 
-		FEditorAnalyticsSession::SaveStoredSessionIDs(SessionIDs);
+			FEditorAnalyticsSession::SaveStoredSessionIDs(SessionIDs);
 
-		FEditorAnalyticsSession::Unlock();
-	}
+			FEditorAnalyticsSession::Unlock();
+			bSessionsLoaded = true;
+		}
+	} while (bForceSendCurrentSession && !bSessionsLoaded); // Retry until session are loaded if the sender is forced to send the current session.
 
 	for (const FEditorAnalyticsSession& Session : SessionsToReport)
 	{
diff --git a/Engine/Source/Developer/EditorAnalyticsSession/Public/EditorAnalyticsSession.h b/Engine/Source/Developer/EditorAnalyticsSession/Public/EditorAnalyticsSession.h
--- a/Engine/Source/Developer/EditorAnalyticsSession/Public/EditorAnalyticsSession.h
+++ b/Engine/Source/Developer/EditorAnalyticsSession/Public/EditorAnalyticsSession.h
@@ -110,6 +110,12 @@ struct EDITORANALYTICSSESSION_API FEditorAnalyticsSession
 	 */
 	static bool SaveStoredSessionIDs(const TArray<FString>& InSessions);
 
+	/**
+	 * Try to acquire the local storage lock without blocking.
+	 * @return true if the lock was acquired successfully.
+	 */
+	static bool TryLock() { return Lock(FTimespan::Zero()); }
+
 	/**
 	 * Acquire a lock for local storage.
 	 * @returns true if the lock was acquired successfully.
diff --git a/Engine/Source/Runtime/Engine/Private/Analytics/EditorSessionSummaryWriter.cpp b/Engine/Source/Runtime/Engine/Private/Analytics/EditorSessionSummaryWriter.cpp
--- a/Engine/Source/Runtime/Engine/Private/Analytics/EditorSessionSummaryWriter.cpp
+++ b/Engine/Source/Runtime/Engine/Private/Analytics/EditorSessionSummaryWriter.cpp
@@ -57,7 +57,7 @@ void FEditorSessionSummaryWriter::Initialize()
 
 	UE_LOG(LogEditorSessionSummary, Verbose, TEXT("Initializing EditorSessionSummaryWriter for editor session tracking"));
 
-	if (FEditorAnalyticsSession::Lock()) // System wide locking to write the session file/registry
+	if (FEditorAnalyticsSession::TryLock()) // System wide lock to write the session file/registry. Don't block if already taken, delay initialisation to the next Tick().
 	{
 		// Create a session Session for this session
 		CurrentSession = CreateCurrentSession();
@@ -123,11 +123,14 @@ void FEditorSessionSummaryWriter::Tick(float DeltaTime)
 		return;
 	}
 
-	// Note: Update idle time in Tick() because Slate cannot be invoked from any thread and UpdateTimeStamps() can be called from a crashing thread.
-	//       Compute the idle time from Slate point-of view. Note that some tasks blocking the UI (such as importing large assets) may be considered idle time.
-	CurrentSession->IdleSeconds = FSlateApplication::Get().GetLastUserInteractionTime() != 0 ? // In case Slate did not register any interaction yet (ex the user just launches the Editor and goes away)
-		FMath::FloorToInt(static_cast<float>(FSlateApplication::Get().GetCurrentTime() - FSlateApplication::Get().GetLastUserInteractionTime())) :
-		FMath::FloorToInt(static_cast<float>((FDateTime::UtcNow() - CurrentSession->StartupTimestamp).GetTotalSeconds()));
+	if (CurrentSession != nullptr)
+	{
+		// Note: Update idle time in Tick() because Slate cannot be invoked from any thread and UpdateTimeStamps() can be called from a crashing thread.
+		//       Compute the idle time from Slate point-of view. Note that some tasks blocking the UI (such as importing large assets) may be considered idle time.
+		CurrentSession->IdleSeconds = FSlateApplication::Get().GetLastUserInteractionTime() != 0 ? // In case Slate did not register any interaction yet (ex the user just launches the Editor and goes away)
+			FMath::FloorToInt(static_cast<float>(FSlateApplication::Get().GetCurrentTime() - FSlateApplication::Get().GetLastUserInteractionTime())) :
+			FMath::FloorToInt(static_cast<float>((FDateTime::UtcNow() - CurrentSession->StartupTimestamp).GetTotalSeconds()));
+	}
 
 	HeartbeatTimeElapsed += DeltaTime;
 
@@ -135,7 +138,7 @@ void FEditorSessionSummaryWriter::Tick(float DeltaTime)
 	{
 		HeartbeatTimeElapsed = 0.0f;
 
-		// Try late initialization (in case user toggled the analytics 'Send Data' on during the session).
+		// Try late initialization (in case the global lock was already taken during init and the session couldn't be created or the user just toggled 'send data' on).
 		if (CurrentSession == nullptr)
 		{
 			Initialize();
@@ -338,7 +341,7 @@ FString FEditorSessionSummaryWriter::GetUserActivityString()
 
 void FEditorSessionSummaryWriter::TrySaveCurrentSession()
 {
-	if (FEditorAnalyticsSession::Lock()) // Inter-process lock to grant this process exclusive access to the key-store file/registry.
+	if (FEditorAnalyticsSession::TryLock()) // Inter-process lock to grant this process exclusive access to the key-store file/registry.
 	{
 		FScopeLock ScopedLock(&SaveSessionLock); // Intra-process lock to grant the calling thread exclusive access to the key-store file/registry.
 		CurrentSession->Save();
