Fixed a variable in analytic session that could be use before initialization.

diff --git a/Engine/Source/Developer/EditorAnalyticsSession/Private/EditorAnalyticsSession.cpp b/Engine/Source/Developer/EditorAnalyticsSession/Private/EditorAnalyticsSession.cpp
--- a/Engine/Source/Developer/EditorAnalyticsSession/Private/EditorAnalyticsSession.cpp
+++ b/Engine/Source/Developer/EditorAnalyticsSession/Private/EditorAnalyticsSession.cpp
@@ -359,6 +359,7 @@ FEditorAnalyticsSession::FEditorAnalyticsSession()
 	MonitorProcessID = 0;
 	StartupTimestamp = FDateTime::MinValue();
 	Timestamp = FDateTime::MinValue();
+	IdleSeconds = 0;
 	Idle1Min = 0;
 	Idle5Min = 0;
 	Idle30Min = 0;
