Fixed a race condition firing an ensure when saving the Editor analytic session in background.

diff --git a/Engine/Source/Developer/EditorAnalyticsSession/Private/EditorAnalyticsSession.cpp b/Engine/Source/Developer/EditorAnalyticsSession/Private/EditorAnalyticsSession.cpp
--- a/Engine/Source/Developer/EditorAnalyticsSession/Private/EditorAnalyticsSession.cpp
+++ b/Engine/Source/Developer/EditorAnalyticsSession/Private/EditorAnalyticsSession.cpp
@@ -470,7 +470,8 @@ namespace EditorAnalyticsUtils
 	}
 }
 
-FSystemWideCriticalSection* FEditorAnalyticsSession::StoredValuesLock = nullptr;
+TUniquePtr<FSystemWideCriticalSection> FEditorAnalyticsSession::StoredValuesLock;
+std::atomic<uint64> FEditorAnalyticsSession::StoredValuesLockOwnerInfo(0);
 
 FEditorAnalyticsSession::FEditorAnalyticsSession()
 {
@@ -512,43 +513,50 @@ FEditorAnalyticsSession::FEditorAnalyticsSession()
 
 bool FEditorAnalyticsSession::Lock(FTimespan Timeout)
 {
-	if (!ensure(!IsLocked()))
+	// If Lock() is called recursively from a thread already owning the lock, fire an ensure. FSystemWideCriticalSection doesn't support recursive locking correctly.
+	if (!ensure(!IsLockedBy(FPlatformTLS::GetCurrentThreadId())))
 	{
 		return true;
 	}
-	
-	StoredValuesLock = new FSystemWideCriticalSection(EditorAnalyticsDefs::GlobalLockName, Timeout);
 
-	if (!IsLocked())
-	{
-		delete StoredValuesLock;
-		StoredValuesLock = nullptr;
+	// Try to acquire the lock.
+	TUniquePtr<FSystemWideCriticalSection> TempLock = MakeUnique<FSystemWideCriticalSection>(EditorAnalyticsDefs::GlobalLockName, Timeout);
 
-		return false;
+	// If the lock is acquired.
+	if (TempLock && TempLock->IsValid())
+	{
+		// The lock was successfully acquired by this thread, keep it until Unlock() is called.
+		StoredValuesLock = MoveTemp(TempLock);
+		static_assert(sizeof(FPlatformTLS::GetCurrentThreadId()) == sizeof(uint32), "Encoding below assume the thread Id only uses 32 bits");
+		StoredValuesLockOwnerInfo.store((1ull << sizeof(uint32) * 8) | FPlatformTLS::GetCurrentThreadId()); // 32 bits for thread id, 1 bit for 'is taken' so that threadId 0 is supported.
+		return true;
 	}
 
-	return true;
+	return false; // Failed to get the lock within the allowed time.
 }
 
 void FEditorAnalyticsSession::Unlock()
 {
-	if (!ensure(IsLocked()))
+	// If Unlock() is called from a non-owning thread or recursively from an owning thread, fire an ensure. FSystemWideCriticalSection doesn't support recursive locking correctly.
+	if (!ensure(IsLockedBy(FPlatformTLS::GetCurrentThreadId())))
 	{
 		return;
 	}
 
-	delete StoredValuesLock;
-	StoredValuesLock = nullptr;
+	StoredValuesLock.Reset();
+	StoredValuesLockOwnerInfo.store(0);
 }
 
-bool FEditorAnalyticsSession::IsLocked()
+bool FEditorAnalyticsSession::IsLockedBy(uint32 ThreadId)
 {
-	return StoredValuesLock != nullptr && StoredValuesLock->IsValid();
+	// Check if the calling thread has the lock.
+	uint64 LockInfo = StoredValuesLockOwnerInfo.load();
+	return LockInfo != 0 && (LockInfo & 0xFFFFFFFF) == ThreadId;
 }
 
 bool FEditorAnalyticsSession::Save()
 {
-	if (!ensure(IsLocked()))
+	if (!ensure(IsLockedBy(FPlatformTLS::GetCurrentThreadId())))
 	{
 		return false;
 	}
@@ -652,7 +660,7 @@ bool FEditorAnalyticsSession::Save()
 
 bool FEditorAnalyticsSession::Load(const FString& InSessionID)
 {
-	if (!ensure(IsLocked()))
+	if (!ensure(IsLockedBy(FPlatformTLS::GetCurrentThreadId())))
 	{
 		return false;
 	}
@@ -665,7 +673,7 @@ bool FEditorAnalyticsSession::Load(const FString& InSessionID)
 
 bool FEditorAnalyticsSession::Delete() const
 {
-	if (!ensure(IsLocked()))
+	if (!ensure(IsLockedBy(FPlatformTLS::GetCurrentThreadId())))
 	{
 		return false;
 	}
@@ -747,7 +755,7 @@ bool FEditorAnalyticsSession::Delete() const
 
 bool FEditorAnalyticsSession::GetStoredSessionIDs(TArray<FString>& OutSessions)
 {
-	if (!ensure(IsLocked()))
+	if (!ensure(IsLockedBy(FPlatformTLS::GetCurrentThreadId())))
 	{
 		return false;
 	}
@@ -758,7 +766,7 @@ bool FEditorAnalyticsSession::GetStoredSessionIDs(TArray<FString>& OutSessions)
 
 bool FEditorAnalyticsSession::LoadAllStoredSessions(TArray<FEditorAnalyticsSession>& OutSessions)
 {
-	if (!ensure(IsLocked()))
+	if (!ensure(IsLockedBy(FPlatformTLS::GetCurrentThreadId())))
 	{
 		return false;
 	}
@@ -791,7 +799,7 @@ bool FEditorAnalyticsSession::SaveStoredSessionIDs(const TArray<FString>& InSess
 		SessionListString.Append(Session);
 	}
 
-	if (!ensure(IsLocked()))
+	if (!ensure(IsLockedBy(FPlatformTLS::GetCurrentThreadId())))
 	{
 		return false;
 	}
@@ -802,7 +810,7 @@ bool FEditorAnalyticsSession::SaveStoredSessionIDs(const TArray<FString>& InSess
 
 void FEditorAnalyticsSession::CleanupOutdatedIncompatibleSessions(const FTimespan& MaxAge)
 {
-	if (!ensure(IsLocked()))
+	if (!ensure(IsLockedBy(FPlatformTLS::GetCurrentThreadId())))
 	{
 		return;
 	}
@@ -855,7 +863,7 @@ void FEditorAnalyticsSession::LogEvent(EEventType InEventType, const FDateTime&
 
 bool FEditorAnalyticsSession::FindSession(const uint32 InSessionProcessId, FEditorAnalyticsSession& OutSession)
 {
-	if (!ensure(IsLocked()))
+	if (!ensure(IsLockedBy(FPlatformTLS::GetCurrentThreadId())))
 	{
 		return false;
 	}
@@ -879,7 +887,7 @@ bool FEditorAnalyticsSession::FindSession(const uint32 InSessionProcessId, FEdit
 
 bool FEditorAnalyticsSession::SaveExitCode(int32 InExitCode, const FDateTime& ApproximativeEditorDeathTime)
 {
-	if (!ensure(IsLocked()))
+	if (!ensure(IsLockedBy(FPlatformTLS::GetCurrentThreadId())))
 	{
 		return false;
 	}
@@ -894,7 +902,7 @@ bool FEditorAnalyticsSession::SaveExitCode(int32 InExitCode, const FDateTime& Ap
 
 bool FEditorAnalyticsSession::SaveMonitorExceptCode(int32 InExceptCode)
 {
-	if (!ensure(IsLocked()))
+	if (!ensure(IsLockedBy(FPlatformTLS::GetCurrentThreadId())))
 	{
 		return false;
 	}
diff --git a/Engine/Source/Developer/EditorAnalyticsSession/Public/EditorAnalyticsSession.h b/Engine/Source/Developer/EditorAnalyticsSession/Public/EditorAnalyticsSession.h
--- a/Engine/Source/Developer/EditorAnalyticsSession/Public/EditorAnalyticsSession.h
+++ b/Engine/Source/Developer/EditorAnalyticsSession/Public/EditorAnalyticsSession.h
@@ -7,6 +7,7 @@
 #include "Misc/DateTime.h"
 #include "Modules/ModuleInterface.h"
 #include "Templates/Atomic.h"
+#include <atomic>
 
 struct EDITORANALYTICSSESSION_API FEditorAnalyticsSession
 {
@@ -160,9 +161,6 @@ struct EDITORANALYTICSSESSION_API FEditorAnalyticsSession
 	 */
 	static void Unlock();
 
-	/** Is the local storage already locked? */
-	static bool IsLocked();
-
 	/**
 	 * Append an event to the session log. The function is meant to record concurrent events, especially during a crash
 	 * with minimum contention. The logger appends and persists the events of interest locklessly on spot as opposed to
@@ -200,7 +198,14 @@ struct EDITORANALYTICSSESSION_API FEditorAnalyticsSession
 	bool SaveMonitorExceptCode(int32 ExceptCode);
 
 private:
-	static FSystemWideCriticalSection* StoredValuesLock;
+	/** Returns true if the local storage is locked by the specified thread. */
+	static bool IsLockedBy(uint32 ThreadId);
+
+	/** Ensure exclusive access to the stored values. */
+	static TUniquePtr<FSystemWideCriticalSection> StoredValuesLock;
+
+	/** Keep track of the lock owner to prevent using FSystemWideCriticalSection recursively, which is not supported by the implementation. */
+	static std::atomic<uint64> StoredValuesLockOwnerInfo;
 
 	/** 
 	 * Has this session already been saved? 
@@ -209,6 +214,7 @@ private:
 	bool bAlreadySaved : 1;
 };
 
+
 class FEditorAnalyticsSessionModule : public IModuleInterface
 {
 
