Fixed static analyis warning V634: The priority of the '*' operation is higher than that of the '<<' operation. It's possible that parentheses should be used in the expression in EditorAnalyticSession.cpp

#rnx

diff --git a/Engine/Source/Developer/EditorAnalyticsSession/Private/EditorAnalyticsSession.cpp b/Engine/Source/Developer/EditorAnalyticsSession/Private/EditorAnalyticsSession.cpp
--- a/Engine/Source/Developer/EditorAnalyticsSession/Private/EditorAnalyticsSession.cpp
+++ b/Engine/Source/Developer/EditorAnalyticsSession/Private/EditorAnalyticsSession.cpp
@@ -528,7 +528,7 @@ bool FEditorAnalyticsSession::Lock(FTimespan Timeout)
 		// The lock was successfully acquired by this thread, keep it until Unlock() is called.
 		StoredValuesLock = MoveTemp(TempLock);
 		static_assert(sizeof(FPlatformTLS::GetCurrentThreadId()) == sizeof(uint32), "Encoding below assume the thread Id only uses 32 bits");
-		StoredValuesLockOwnerInfo.store((1ull << sizeof(uint32) * 8) | FPlatformTLS::GetCurrentThreadId()); // 32 bits for thread id, 1 bit for 'is taken' so that threadId 0 is supported.
+		StoredValuesLockOwnerInfo.store(((1ull << (sizeof(uint32) * 8))) | FPlatformTLS::GetCurrentThreadId()); // 32 bits for thread id, 1 bit for 'is taken' so that threadId 0 is supported.
 		return true;
 	}
 
