- Count number of crashes before Analytics get initialized and report them as DelayedCrashCount field of the Editor summary session event.

diff --git a/Engine/Source/Developer/EditorAnalyticsSession/Private/EditorAnalyticsSession.cpp b/Engine/Source/Developer/EditorAnalyticsSession/Private/EditorAnalyticsSession.cpp
--- a/Engine/Source/Developer/EditorAnalyticsSession/Private/EditorAnalyticsSession.cpp
+++ b/Engine/Source/Developer/EditorAnalyticsSession/Private/EditorAnalyticsSession.cpp
@@ -5,6 +5,7 @@
 #include "Internationalization/Regex.h"
 #include "HAL/PlatformProcess.h"
 #include "HAL/FileManager.h"
+#include "Misc/EngineVersion.h"
 #include "Misc/Paths.h"
 
 IMPLEMENT_MODULE(FEditorAnalyticsSessionModule, EditorAnalyticsSession);
@@ -480,18 +481,9 @@ FEditorAnalyticsSession::FEditorAnalyticsSession()
 	UserId = EditorAnalyticsDefs::UnknownUserIdString;
 
 	ProjectName = EditorAnalyticsDefs::UnknownProjectValueString;
-	PlatformProcessID = 0;
-	MonitorProcessID = 0;
 	StartupTimestamp = FDateTime::MinValue();
 	Timestamp = FDateTime::MinValue();
 	CurrentUserActivity = EditorAnalyticsDefs::DefaultUserActivity;
-	AverageFPS = 0;
-	GPUVendorID = 0;
-	GPUDeviceID = 0;
-	GRHIDeviceRevision = 0;
-	TotalPhysicalRAM = 0;
-	CPUPhysicalCores = 0;
-	CPULogicalCores = 0;
 
 	bIs64BitOS = false;
 	bCrashed = false;
@@ -912,3 +904,32 @@ bool FEditorAnalyticsSession::SaveMonitorExceptCode(int32 InExceptCode)
 	const FString StorageLocation = EditorAnalyticsUtils::GetSessionStorageLocation(SessionId);
 	return FPlatformMisc::SetStoredValue(EditorAnalyticsDefs::StoreId, StorageLocation, EditorAnalyticsDefs::MonitorExceptCodeStoreKey, FString::FromInt(InExceptCode));
 }
+
+void FEditorAnalyticsSession::CreateMinimalCrashSession(const TOptional<int32>& ExitCode)
+{
+	if (!ensure(IsLockedBy(FPlatformTLS::GetCurrentThreadId()))) // Caller needs to own the session lock.
+	{
+		return;
+	}
+
+	FEditorAnalyticsSession MinimalSession;
+	MinimalSession.SessionId = FGuid::NewGuid().ToString(EGuidFormats::DigitsWithHyphens); // To have a unique session id used as key in the storage.
+	MinimalSession.bCrashed = true;
+	MinimalSession.ExitCode = ExitCode;
+	MinimalSession.StartupTimestamp = FDateTime::UtcNow();
+	MinimalSession.Timestamp = MinimalSession.StartupTimestamp; // To enable session 'expiration'.
+	MinimalSession.EngineVersion = FEngineVersion::Current().ToString(EVersionComponent::Changelist); // To account the crash for the corresponding engine version.
+
+	MinimalSession.Save();
+
+	// Update the session list.
+	TArray<FString> StoredSessions;
+	GetStoredSessionIDs(StoredSessions);
+	StoredSessions.Add(MinimalSession.SessionId);
+	SaveStoredSessionIDs(StoredSessions);
+}
+
+bool FEditorAnalyticsSession::IsMinimalCrashSession() const
+{
+	return AppId == EditorAnalyticsDefs::UnknownAppIdString && ProjectID.IsEmpty() && Plugins.Num() == 0; // Check fields values that are not expected when the session is fully initialized.
+}
diff --git a/Engine/Source/Developer/EditorAnalyticsSession/Private/EditorSessionSummarySender.cpp b/Engine/Source/Developer/EditorAnalyticsSession/Private/EditorSessionSummarySender.cpp
--- a/Engine/Source/Developer/EditorAnalyticsSession/Private/EditorSessionSummarySender.cpp
+++ b/Engine/Source/Developer/EditorAnalyticsSession/Private/EditorSessionSummarySender.cpp
@@ -29,6 +29,17 @@ namespace EditorSessionSenderDefs
 	static const FString AbnormalSessionToken(TEXT("AbnormalShutdown"));
 }
 
+/** Contains the analytics session and some extra data to report. */
+struct FEditorAnalyticReportData
+{
+	/** The analytic session to report. */
+	FEditorAnalyticsSession Session;
+
+	/** Number of Editor instances (same engine version than Session) that crashed prior the analytic was running and wasn't accounted for yet. */
+	uint32 DelayedCrashCount = 0;
+};
+
+
 FEditorSessionSummarySender::FEditorSessionSummarySender(IAnalyticsProviderET& InAnalyticsProvider, const FString& InSenderName, const uint32 InCurrentSessionProcessId)
 	: HeartbeatTimeElapsed(0.0f)
 	, AnalyticsProvider(InAnalyticsProvider)
@@ -68,7 +79,7 @@ void FEditorSessionSummarySender::SendStoredSessions(const bool bForceSendOwnedS
 	// Load the list of sessions to process. Expect contention on the analytic session lock between the Editor and CrashReportClientEditor (on windows) or between Editor instances (on mac/linux)
 	//   - Try every 'n' seconds if bForceSendCurrentSession is true.
 	//   - Don't block and don't loop if bForceSendCurrentSession is false.
-	TArray<FEditorAnalyticsSession> SessionsToReport;
+	TArray<FEditorAnalyticReportData> SessionsToReport;
 	FTimespan Timemout(bForceSendOwnedSession ? FTimespan::FromSeconds(5) : FTimespan::Zero());
 	bool bSessionsLoaded = false;
 	do
@@ -77,6 +88,7 @@ void FEditorSessionSummarySender::SendStoredSessions(const bool bForceSendOwnedS
 		{
 			// Get list of sessions in storage
 			TArray<FEditorAnalyticsSession> ExistingSessions;
+			TArray<FEditorAnalyticsSession> MinimalCrashSessions;
 			FEditorAnalyticsSession::LoadAllStoredSessions(ExistingSessions);
 
 			TArray<FEditorAnalyticsSession> SessionsToDelete;
@@ -88,7 +100,7 @@ void FEditorSessionSummarySender::SendStoredSessions(const bool bForceSendOwnedS
 			auto IsOrphan = [](const FEditorAnalyticsSession& InSession) { return !FPlatformProcess::IsApplicationRunning(InSession.PlatformProcessID) && (InSession.MonitorProcessID == 0 || !FPlatformProcess::IsApplicationRunning(InSession.MonitorProcessID)); };
 
 			// Check each stored session to see if they should be sent or not
-			for (FEditorAnalyticsSession& Session : ExistingSessions)
+			for (const FEditorAnalyticsSession& Session : ExistingSessions)
 			{
 				if (HasOwnershipOf(Session)) // This process was configured to send this session.
 				{
@@ -105,11 +117,31 @@ void FEditorSessionSummarySender::SendStoredSessions(const bool bForceSendOwnedS
 				const FTimespan SessionAge = FDateTime::UtcNow() - Session.Timestamp;
 				if (SessionAge < EditorSessionSenderDefs::SessionExpiration)
 				{
-					SessionsToReport.Add(Session);
+					if (Session.IsMinimalCrashSession())
+					{
+						MinimalCrashSessions.Add(Session);
+						continue; // Don't schedule this minimal session for reporting/deletiong yet, it's data must be piggybacked off a full and compatible session to be reported/deleted.
+					}
+
+					SessionsToReport.Add(FEditorAnalyticReportData{Session,0});
 				}
 				SessionsToDelete.Add(Session);
 			}
 
+			// Try to piggyback the minimal crash session information off a valid session. Minimal sessions are incomplete and will be rejected/ignored by the backend or may disrupt various analysis.
+			for (FEditorAnalyticsSession& MinimalCrashSession : MinimalCrashSessions)
+			{
+				// Piggyback information of the minimal crash session off a session that has the same engine version. (Stats are aggregated by engine version)
+				if (FEditorAnalyticReportData* ReportData = SessionsToReport.FindByPredicate([&MinimalCrashSession](const FEditorAnalyticReportData& Candidate) { return Candidate.Session.EngineVersion == MinimalCrashSession.EngineVersion; }))
+				{
+					// Account for the crashes that occurred before the analytic was initialized.
+					ReportData->DelayedCrashCount += 1;
+
+					// The minimal crash session was accounted for, schedule it for deletion.
+					SessionsToDelete.Add(MoveTemp(MinimalCrashSession));
+				}
+			}
+
 			// NOTE: The sessions are deleted before sending (while holding the lock) to prevent another process from finding and sending the same orphan sessions twice.
 			//       This also ensures a session will never be sent twice in case a crash occurs while sending, but all sessions loaded in memory for sending will also be lost.
 			for (const FEditorAnalyticsSession& ToDelete : SessionsToDelete)
@@ -136,14 +168,15 @@ void FEditorSessionSummarySender::SendStoredSessions(const bool bForceSendOwnedS
 	} while (bForceSendOwnedSession && !bSessionsLoaded); // Retry until session are loaded if the sender is forced to send the current session.
 
 	// Send the sessions (without holding the lock).
-	for (const FEditorAnalyticsSession& Session : SessionsToReport)
+	for (const FEditorAnalyticReportData& ReportData : SessionsToReport)
 	{
-		SendSessionSummaryEvent(Session);
+		SendSessionSummaryEvent(ReportData);
 	}
 }
 
-void FEditorSessionSummarySender::SendSessionSummaryEvent(const FEditorAnalyticsSession& Session) const
+void FEditorSessionSummarySender::SendSessionSummaryEvent(const FEditorAnalyticReportData& ReportData) const
 {
+	const FEditorAnalyticsSession& Session = ReportData.Session;
 	FGuid SessionId;
 	FString SessionIdString = Session.SessionId;
 	if (FGuid::Parse(SessionIdString, SessionId))
@@ -222,6 +255,7 @@ void FEditorSessionSummarySender::SendSessionSummaryEvent(const FEditorAnalytics
 	AnalyticsAttributes.Emplace(TEXT("IsCrcMissing"), Session.bIsCrcExeMissing);
 	AnalyticsAttributes.Emplace(TEXT("SentFrom"), Sender);
 	AnalyticsAttributes.Emplace(TEXT("MonitorPid"), Session.MonitorProcessID); // For out-of-process monitoring, if this is 0, this will mean CRC failed to launch or crashed very early.
+	AnalyticsAttributes.Emplace(TEXT("DelayedCrashCount"), ReportData.DelayedCrashCount); // Piggybacked information from 'minimal crash sessions', previous Editor instances crashes recorded before analytics was initialized.
 
 	bool bShouldAttachMonitorLog = (ShutdownTypeString != EditorSessionSenderDefs::ShutdownSessionToken && ShutdownTypeString != EditorSessionSenderDefs::TerminatedSessionToken);
 
diff --git a/Engine/Source/Developer/EditorAnalyticsSession/Public/EditorAnalyticsSession.h b/Engine/Source/Developer/EditorAnalyticsSession/Public/EditorAnalyticsSession.h
--- a/Engine/Source/Developer/EditorAnalyticsSession/Public/EditorAnalyticsSession.h
+++ b/Engine/Source/Developer/EditorAnalyticsSession/Public/EditorAnalyticsSession.h
@@ -31,8 +31,8 @@ struct EDITORANALYTICSSESSION_API FEditorAnalyticsSession
 	FString ProjectDescription;
 	FString ProjectVersion;
 	FString EngineVersion;
-	uint32 PlatformProcessID;
-	uint32 MonitorProcessID; // Set to the CrashReportClientEditor PID when out-of-process reporting is used.
+	uint32 PlatformProcessID = 0;
+	uint32 MonitorProcessID = 0; // Set to the CrashReportClientEditor PID when out-of-process reporting is used.
 	TOptional<int32> ExitCode; // Set by CrashReportClientEditor after the Editor process exit when out-of-process reporting is used and reading the exit code is supported.
 	TOptional<int32> MonitorExceptCode; // Set in CrashReportClientEditor process when an exception or an error is caught.
 	TOptional<int32> MonitorExitCode; // Set in the Editor process when the Editor detects that CrashReportClientEditor process unexpectedly died.
@@ -49,7 +49,7 @@ struct EDITORANALYTICSSESSION_API FEditorAnalyticsSession
 	volatile int32 TotalEditorInactivitySeconds = 0; // Account for user input and Editor process CPU usage. Add up gaps where the CPU was not used intensively and the user did not interact.
 	FString CurrentUserActivity;
 	TArray<FString> Plugins;
-	float AverageFPS;
+	float AverageFPS = 0.0f;
 
 	uint64 SessionTickCount = 0; // Number of times the analytic session was ticked. Zero is the interesting value. If the Editor is hang during boot, some users may be prompt to kill it.
 	uint64 EngineTickCount = 0;  // Number or times the engine was ticked.
@@ -57,16 +57,16 @@ struct EDITORANALYTICSSESSION_API FEditorAnalyticsSession
 
 	FString DesktopGPUAdapter;
 	FString RenderingGPUAdapter;
-	uint32 GPUVendorID;
-	uint32 GPUDeviceID;
-	uint32 GRHIDeviceRevision;
+	uint32 GPUVendorID = 0;
+	uint32 GPUDeviceID = 0;
+	uint32 GRHIDeviceRevision = 0;
 	FString GRHIAdapterInternalDriverVersion;
 	FString GRHIAdapterUserDriverVersion;
 	FString GRHIName;
 
-	uint64 TotalPhysicalRAM;
-	int32 CPUPhysicalCores;
-	int32 CPULogicalCores;
+	uint64 TotalPhysicalRAM = 0;
+	int32 CPUPhysicalCores = 0;
+	int32 CPULogicalCores = 0;
 	FString CPUVendor;
 	FString CPUBrand;
 
@@ -197,6 +197,19 @@ struct EDITORANALYTICSSESSION_API FEditorAnalyticsSession
 	 */
 	bool SaveMonitorExceptCode(int32 ExceptCode);
 
+	/**
+	 * Create a very minimal session for an Editor crash that occurred before the analytic session was created. The
+	 * session lifecyle is the same as a normal session, but it cannot be sent 'as-it' to the backend because it is
+	 * incomplete. Instead, the session information is pulled and piggybacked with another valid session. Usually, a
+	 * minimal session will be created by CrashReportClient (in monitor mode) if the Editor crashed and the corresponding
+	 * analytic session couldn't be found.
+	 * @note The purpose of this functionality is to allow counting very early crashes that aren't reported to analytics otherwise.
+	 */
+	static void CreateMinimalCrashSession(const TOptional<int32>& ExitCode);
+	
+	/** Returns true if the session is minimal, i.e. created with CreateMinimalCrashSession(). */
+	bool IsMinimalCrashSession() const;
+
 private:
 	/** Returns true if the local storage is locked by the specified thread. */
 	static bool IsLockedBy(uint32 ThreadId);
diff --git a/Engine/Source/Developer/EditorAnalyticsSession/Public/EditorSessionSummarySender.h b/Engine/Source/Developer/EditorAnalyticsSession/Public/EditorSessionSummarySender.h
--- a/Engine/Source/Developer/EditorAnalyticsSession/Public/EditorSessionSummarySender.h
+++ b/Engine/Source/Developer/EditorAnalyticsSession/Public/EditorSessionSummarySender.h
@@ -11,6 +11,7 @@
 
 struct FEditorAnalyticsSession;
 class IAnalyticsProviderET;
+struct FEditorAnalyticReportData;
 
 /**
   * Sender of SessionSummary events from all editor sessions in-between runs.
@@ -30,7 +31,7 @@ public:
 private:
 	/** Send any stored Sessions. */
 	void SendStoredSessions(const bool bForceSendCurrentSession = false) const;
-	void SendSessionSummaryEvent(const FEditorAnalyticsSession& Session) const;
+	void SendSessionSummaryEvent(const FEditorAnalyticReportData& ReportData) const;
 
 private:
 	float HeartbeatTimeElapsed;
diff --git a/Engine/Source/Programs/CrashReportClient/Private/CrashReportClientApp.cpp b/Engine/Source/Programs/CrashReportClient/Private/CrashReportClientApp.cpp
--- a/Engine/Source/Programs/CrashReportClient/Private/CrashReportClientApp.cpp
+++ b/Engine/Source/Programs/CrashReportClient/Private/CrashReportClientApp.cpp
@@ -89,6 +89,9 @@ static void* MonitorReadPipe = nullptr;
 /** If in monitor mode, pipe to write data to game. */
 static void* MonitorWritePipe = nullptr;
 
+/** If in monitor mode, set to true when the monitored app crashes. */
+static bool bMonitoredAppCrashed = false;
+
 /** Result of submission of report */
 enum SubmitCrashReportResult {
 	Failed,				// Failed to send report
@@ -633,14 +636,42 @@ bool IsCrashReportAvailable(uint32 WatchedProcess, FSharedCrashContext& CrashCon
 			// Record the history of events sent by the Editor to help diagnose abnormal terminations.
 			switch (CrashContext.CrashType)
 			{
-				case ECrashContextType::Assert:           FDiagnosticLogger::Get().LogEvent(TEXT("Pipe/Assert"));          break;
-				case ECrashContextType::Ensure:           FDiagnosticLogger::Get().LogEvent(TEXT("Pipe/Ensure"));          break;
-				case ECrashContextType::Crash:            FDiagnosticLogger::Get().LogEvent(TEXT("Pipe/Crash"));           break;
-				case ECrashContextType::GPUCrash:         FDiagnosticLogger::Get().LogEvent(TEXT("Pipe/GPUCrash"));        break;
-				case ECrashContextType::Hang:             FDiagnosticLogger::Get().LogEvent(TEXT("Pipe/Hang"));            break;
-				case ECrashContextType::OutOfMemory:      FDiagnosticLogger::Get().LogEvent(TEXT("Pipe/OOM"));             break;
-				case ECrashContextType::AbnormalShutdown: FDiagnosticLogger::Get().LogEvent(TEXT("Pipe/AbnormalShutdown"));break;
-				default:                                  FDiagnosticLogger::Get().LogEvent(TEXT("Pipe/Unknown"));         break;
+				case ECrashContextType::Assert:
+					FDiagnosticLogger::Get().LogEvent(TEXT("Pipe/Assert"));
+					bMonitoredAppCrashed = true;
+					break;
+
+				case ECrashContextType::Ensure:
+					FDiagnosticLogger::Get().LogEvent(TEXT("Pipe/Ensure"));
+					break;
+
+				case ECrashContextType::Crash:
+					FDiagnosticLogger::Get().LogEvent(TEXT("Pipe/Crash"));
+					bMonitoredAppCrashed = true;
+					break;
+
+				case ECrashContextType::GPUCrash:
+					FDiagnosticLogger::Get().LogEvent(TEXT("Pipe/GPUCrash"));
+					bMonitoredAppCrashed = true;
+					break;
+
+				case ECrashContextType::Hang:
+					FDiagnosticLogger::Get().LogEvent(TEXT("Pipe/Hang"));
+					break;
+
+				case ECrashContextType::OutOfMemory:
+					FDiagnosticLogger::Get().LogEvent(TEXT("Pipe/OOM"));
+					bMonitoredAppCrashed = true;
+					break;
+
+				case ECrashContextType::AbnormalShutdown:
+					FDiagnosticLogger::Get().LogEvent(TEXT("Pipe/AbnormalShutdown"));
+					bMonitoredAppCrashed = true;
+					break;
+
+				default:
+					FDiagnosticLogger::Get().LogEvent(TEXT("Pipe/Unknown"));
+					break;
 			}
 		}
 
@@ -1074,6 +1105,11 @@ void RunCrashReportClient(const TCHAR* CommandLine)
 					else
 					{
 						FDiagnosticLogger::Get().LogEvent(TEXT("MTBF/NoSessionFound"));
+						if (bMonitoredAppCrashed)
+						{
+							// No session found. Record this crash to piggyback off a valid session later. (Purpose: Account for early Editor crash happening before Analytics was initialized)
+							FEditorAnalyticsSession::CreateMinimalCrashSession(MonitoredProcessExitCode);
+						}
 					}
 					FEditorAnalyticsSession::Unlock();
 					break;
