Fixed CIS errors reported on Android.

diff --git a/Engine/Source/Developer/EditorAnalyticsSession/Private/EditorAnalyticsSession.cpp b/Engine/Source/Developer/EditorAnalyticsSession/Private/EditorAnalyticsSession.cpp
--- a/Engine/Source/Developer/EditorAnalyticsSession/Private/EditorAnalyticsSession.cpp
+++ b/Engine/Source/Developer/EditorAnalyticsSession/Private/EditorAnalyticsSession.cpp
@@ -6,6 +6,7 @@
 #include "HAL/PlatformProcess.h"
 #include "HAL/FileManager.h"
 #include "Misc/EngineVersion.h"
+#include "Misc/Guid.h"
 #include "Misc/Paths.h"
 
 IMPLEMENT_MODULE(FEditorAnalyticsSessionModule, EditorAnalyticsSession);
