Add StallDetector to Session Summary for Analytics

This adds some scalar statistics to the session summary:
* The total stall trigger count
* The total reported stalls
* The top stall name
* The top stall budget time in seconds
* The top stall cumulative overage time in seconds
* The top stall trigger count

diff --git a/Engine/Source/Developer/EditorAnalyticsSession/Private/EditorAnalyticsSession.cpp b/Engine/Source/Developer/EditorAnalyticsSession/Private/EditorAnalyticsSession.cpp
--- a/Engine/Source/Developer/EditorAnalyticsSession/Private/EditorAnalyticsSession.cpp
+++ b/Engine/Source/Developer/EditorAnalyticsSession/Private/EditorAnalyticsSession.cpp
@@ -29,13 +29,15 @@ namespace EditorAnalyticsDefs
 	//   Version 1_2 : Removed TotalUserInactivitySeconds and added SessionDuration.
 	//   Version 1_3 : Added SessionTickCount, UserInteractionCount, IsCrcExeMissing, IsUserLoggingOut, MonitorExitCode and readded lost code to save/load/delete IsLowDriveSpace for 4.26.0.
 	//   Version 1_4 : Added CommandLine, EngineTickCount, LastTickTimestamp, DeathTimestamp and IsDebuggerIgnored for 4.26.0.
+	//   Version 1_5 : Added Stall Detector stats for 5.0.
 	static const FString StoreId(TEXT("Epic Games"));
 	static const FString SessionSummaryRoot(TEXT("Unreal Engine/Session Summary"));
 	static const FString SessionSummarySection_1_0 = SessionSummaryRoot / TEXT("1_0"); // The session format used by older versions.
 	static const FString SessionSummarySection_1_1 = SessionSummaryRoot / TEXT("1_1");
 	static const FString SessionSummarySection_1_2 = SessionSummaryRoot / TEXT("1_2");
 	static const FString SessionSummarySection_1_3 = SessionSummaryRoot / TEXT("1_3");
-	static const FString SessionSummarySection = SessionSummaryRoot / TEXT("1_4"); // The current session format.
+	static const FString SessionSummarySection_1_4 = SessionSummaryRoot / TEXT("1_4");
+	static const FString SessionSummarySection = SessionSummaryRoot / TEXT("1_5"); // The current session format.
 	static const FString GlobalLockName(TEXT("UE4_SessionSummary_Lock"));
 	static const FString SessionListStoreKey(TEXT("SessionList"));
 
@@ -60,6 +62,14 @@ namespace EditorAnalyticsDefs
 	static const FString UserInteractionCountStoreKey(TEXT("UserInteractionCount"));
 	static const FString CommandLineStoreKey(TEXT("CommandLine"));
 
+	// stall detector
+	static const FString TotalStallCountStoreKey(TEXT("TotalStallCount"));
+	static const FString TotalStallReportedStoreKey(TEXT("TotalStallReported"));
+	static const FString TopStallNameStoreKey(TEXT("TopStallName"));
+	static const FString TopStallBudgetSecondsStoreKey(TEXT("TopStallBudgetSeconds"));
+	static const FString TopStallOverageSecondsStoreKey(TEXT("TopStallOverageSeconds"));
+	static const FString TopStallTriggerCountStoreKey(TEXT("TopStallTriggerCount"));
+
 	// timestamps
 	static const FString StartupTimestampStoreKey(TEXT("StartupTimestamp"));
 	static const FString TimestampStoreKey(TEXT("Timestamp"));
@@ -308,7 +318,11 @@ namespace EditorAnalyticsUtils
 	}
 
 // Utility macros to make it easier to check that all fields are being written to.
-#define GET_STORED_STRING(FieldName) FPlatformMisc::GetStoredValue(EditorAnalyticsDefs::StoreId, SectionName, EditorAnalyticsDefs:: FieldName ## StoreKey, Session.FieldName)
+#define GET_STORED_STRING(FieldName) \
+	{ \
+		FPlatformMisc::GetStoredValue(EditorAnalyticsDefs::StoreId, SectionName, EditorAnalyticsDefs:: FieldName ## StoreKey, Session.FieldName); \
+	}
+
 #define GET_STORED_INT(FieldName) \
 	{ \
 		FString FieldName ## Temp; \
@@ -316,6 +330,13 @@ namespace EditorAnalyticsUtils
 		Session.FieldName = FCString::Atoi(*FieldName ## Temp); \
 	}
 
+#define GET_STORED_FLOAT(FieldName) \
+	{ \
+		FString FieldName ## Temp; \
+		FPlatformMisc::GetStoredValue(EditorAnalyticsDefs::StoreId, SectionName, EditorAnalyticsDefs:: FieldName ## StoreKey, FieldName ## Temp); \
+		Session.FieldName = FCString::Atof(*FieldName ## Temp); \
+	}
+
 	static void LoadInternal(FEditorAnalyticsSession& Session, const FString& InSessionId)
 	{
 		Session.SessionId = InSessionId;
@@ -400,16 +421,19 @@ namespace EditorAnalyticsUtils
 			PluginsString.ParseIntoArray(Session.Plugins, TEXT(","));
 		}
 
-		{
-			FString AverageFPSString;
-			FPlatformMisc::GetStoredValue(EditorAnalyticsDefs::StoreId, SectionName, EditorAnalyticsDefs::AverageFPSStoreKey, AverageFPSString);
-			Session.AverageFPS = FCString::Atof(*AverageFPSString);
-		}
+		GET_STORED_FLOAT(AverageFPS);
 
 		GET_STORED_INT(SessionTickCount);
 		GET_STORED_INT(EngineTickCount);
 		GET_STORED_INT(UserInteractionCount);
 
+		GET_STORED_INT(TotalStallCount);
+		GET_STORED_INT(TotalStallReported);
+		GET_STORED_STRING(TopStallName);
+		GET_STORED_FLOAT(TopStallBudgetSeconds);
+		GET_STORED_FLOAT(TopStallOverageSeconds);
+		GET_STORED_INT(TopStallTriggerCount);
+
 		GET_STORED_STRING(DesktopGPUAdapter);
 		GET_STORED_STRING(RenderingGPUAdapter);
 
@@ -613,6 +637,12 @@ bool FEditorAnalyticsSession::Save()
 			{EditorAnalyticsDefs::SessionTickCountStoreKey, FString::FromInt(SessionTickCount)},
 			{EditorAnalyticsDefs::EngineTickCountStoreKey,  FString::FromInt(EngineTickCount)},
 			{EditorAnalyticsDefs::UserInteractionCountStoreKey, FString::FromInt(UserInteractionCount)},
+			{EditorAnalyticsDefs::TotalStallCountStoreKey,        FString::FromInt(TotalStallCount)},
+			{EditorAnalyticsDefs::TotalStallReportedStoreKey,     FString::FromInt(TotalStallReported)},
+			{EditorAnalyticsDefs::TopStallNameStoreKey,           TopStallName},
+			{EditorAnalyticsDefs::TopStallBudgetSecondsStoreKey,  FString::SanitizeFloat(TopStallBudgetSeconds)},
+			{EditorAnalyticsDefs::TopStallOverageSecondsStoreKey, FString::SanitizeFloat(TopStallOverageSeconds)},
+			{EditorAnalyticsDefs::TopStallTriggerCountStoreKey,   FString::FromInt(TopStallTriggerCount)},
 			{EditorAnalyticsDefs::IsDebuggerStoreKey,       EditorAnalyticsUtils::BoolToStoredString(bIsDebugger)},
 			{EditorAnalyticsDefs::IsDebuggerIgnoredStoreKey,EditorAnalyticsUtils::BoolToStoredString(bIsDebuggerIgnored)},
 			{EditorAnalyticsDefs::WasDebuggerStoreKey,      EditorAnalyticsUtils::BoolToStoredString(bWasEverDebugger)},
@@ -705,6 +735,13 @@ bool FEditorAnalyticsSession::Delete() const
 	FPlatformMisc::DeleteStoredValue(EditorAnalyticsDefs::StoreId, SectionName, EditorAnalyticsDefs::UserInteractionCountStoreKey);
 	FPlatformMisc::DeleteStoredValue(EditorAnalyticsDefs::StoreId, SectionName, EditorAnalyticsDefs::CommandLineStoreKey);
 
+	FPlatformMisc::DeleteStoredValue(EditorAnalyticsDefs::StoreId, SectionName, EditorAnalyticsDefs::TotalStallCountStoreKey);
+	FPlatformMisc::DeleteStoredValue(EditorAnalyticsDefs::StoreId, SectionName, EditorAnalyticsDefs::TotalStallReportedStoreKey);
+	FPlatformMisc::DeleteStoredValue(EditorAnalyticsDefs::StoreId, SectionName, EditorAnalyticsDefs::TopStallNameStoreKey);
+	FPlatformMisc::DeleteStoredValue(EditorAnalyticsDefs::StoreId, SectionName, EditorAnalyticsDefs::TopStallBudgetSecondsStoreKey);
+	FPlatformMisc::DeleteStoredValue(EditorAnalyticsDefs::StoreId, SectionName, EditorAnalyticsDefs::TopStallOverageSecondsStoreKey);
+	FPlatformMisc::DeleteStoredValue(EditorAnalyticsDefs::StoreId, SectionName, EditorAnalyticsDefs::TopStallTriggerCountStoreKey);
+
 	FPlatformMisc::DeleteStoredValue(EditorAnalyticsDefs::StoreId, SectionName, EditorAnalyticsDefs::DesktopGPUAdapterStoreKey);
 	FPlatformMisc::DeleteStoredValue(EditorAnalyticsDefs::StoreId, SectionName, EditorAnalyticsDefs::RenderingGPUAdapterStoreKey);
 	FPlatformMisc::DeleteStoredValue(EditorAnalyticsDefs::StoreId, SectionName, EditorAnalyticsDefs::GPUVendorIDStoreKey);
@@ -847,6 +884,7 @@ void FEditorAnalyticsSession::CleanupOutdatedIncompatibleSessions(const FTimespa
 	CleanupVersionedSection(EditorAnalyticsDefs::SessionSummarySection_1_1);
 	CleanupVersionedSection(EditorAnalyticsDefs::SessionSummarySection_1_2);
 	CleanupVersionedSection(EditorAnalyticsDefs::SessionSummarySection_1_3);
+	CleanupVersionedSection(EditorAnalyticsDefs::SessionSummarySection_1_4);
 }
 
 void FEditorAnalyticsSession::LogEvent(EEventType InEventType, const FDateTime& InTimestamp)
diff --git a/Engine/Source/Developer/EditorAnalyticsSession/Private/EditorSessionSummarySender.cpp b/Engine/Source/Developer/EditorAnalyticsSession/Private/EditorSessionSummarySender.cpp
--- a/Engine/Source/Developer/EditorAnalyticsSession/Private/EditorSessionSummarySender.cpp
+++ b/Engine/Source/Developer/EditorAnalyticsSession/Private/EditorSessionSummarySender.cpp
@@ -224,6 +224,12 @@ void FEditorSessionSummarySender::SendSessionSummaryEvent(const FEditorAnalyticR
 	AnalyticsAttributes.Emplace(TEXT("SessionTickCount"), Session.SessionTickCount);
 	AnalyticsAttributes.Emplace(TEXT("EngineTickCount"), Session.EngineTickCount);
 	AnalyticsAttributes.Emplace(TEXT("UserInteractionCount"), Session.UserInteractionCount);
+	AnalyticsAttributes.Emplace(TEXT("TotalStallCount"), Session.TotalStallCount);
+	AnalyticsAttributes.Emplace(TEXT("TotalStallReported"), Session.TotalStallReported);
+	AnalyticsAttributes.Emplace(TEXT("TopStallName"), Session.TopStallName);
+	AnalyticsAttributes.Emplace(TEXT("TopStallBudgetSeconds"), Session.TopStallBudgetSeconds);
+	AnalyticsAttributes.Emplace(TEXT("TopStallOverageSeconds"), Session.TopStallOverageSeconds);
+	AnalyticsAttributes.Emplace(TEXT("TopStallTriggerCount"), Session.TopStallTriggerCount);
 	AnalyticsAttributes.Emplace(TEXT("Plugins"), PluginsString);
 	AnalyticsAttributes.Emplace(TEXT("DesktopGPUAdapter"), Session.DesktopGPUAdapter);
 	AnalyticsAttributes.Emplace(TEXT("RenderingGPUAdapter"), Session.RenderingGPUAdapter);
diff --git a/Engine/Source/Developer/EditorAnalyticsSession/Public/EditorAnalyticsSession.h b/Engine/Source/Developer/EditorAnalyticsSession/Public/EditorAnalyticsSession.h
--- a/Engine/Source/Developer/EditorAnalyticsSession/Public/EditorAnalyticsSession.h
+++ b/Engine/Source/Developer/EditorAnalyticsSession/Public/EditorAnalyticsSession.h
@@ -55,6 +55,13 @@ struct EDITORANALYTICSSESSION_API FEditorAnalyticsSession
 	uint64 EngineTickCount = 0;  // Number or times the engine was ticked.
 	uint32 UserInteractionCount = 0; // Number of slate user interactions. Zero is the interesting value. If the Editor UI hang at start up, some users may be prompt to kill it.
 
+	uint64 TotalStallCount = 0;    // The total number of stalls detected, across every detector
+	uint64 TotalStallReported = 0; // The total number of stalls reported (by default we only report the first instance, or skip reporting some detectors entirely)
+	FString TopStallName;
+	double TopStallBudgetSeconds = 0.0;
+	double TopStallOverageSeconds = 0.0;
+	uint64 TopStallTriggerCount = 0;
+
 	FString DesktopGPUAdapter;
 	FString RenderingGPUAdapter;
 	uint32 GPUVendorID = 0;
diff --git a/Engine/Source/Runtime/Core/Private/ProfilingDebugging/StallDetector.cpp b/Engine/Source/Runtime/Core/Private/ProfilingDebugging/StallDetector.cpp
--- a/Engine/Source/Runtime/Core/Private/ProfilingDebugging/StallDetector.cpp
+++ b/Engine/Source/Runtime/Core/Private/ProfilingDebugging/StallDetector.cpp
@@ -107,7 +107,7 @@ uint32 UE::FStallDetectorRunnable::Run()
 
 		// Check the detectors
 		{
-			FScopeLock ScopeLock(&FStallDetector::GetCriticalSection());
+			FScopeLock ScopeLock(&FStallDetector::GetInstancesSection());
 			for (FStallDetector* Detector : FStallDetector::GetInstances())
 			{
 				Detector->Check(false, Seconds);
@@ -125,7 +125,7 @@ uint32 UE::FStallDetectorRunnable::Run()
 * Stall Detector Stats
 **/
 
-FCriticalSection UE::FStallDetectorStats::CriticalSection;
+FCriticalSection UE::FStallDetectorStats::InstancesSection;
 TSet<UE::FStallDetectorStats*> UE::FStallDetectorStats::Instances;
 std::atomic<uint32> UE::FStallDetectorStats::TotalTriggeredCount;
 std::atomic<uint32> UE::FStallDetectorStats::TotalReportedCount;
@@ -134,26 +134,91 @@ UE::FStallDetectorStats::FStallDetectorStats(const TCHAR* InName, const double I
 	: Name(InName)
 	, BudgetSeconds(InBudgetSeconds)
 	, ReportingMode(InReportingMode)
+	, bReported(false)
 	, TriggerCount(0)
 	, OverageSeconds(0.0)
 {
 	// Add at the end of construction
-	FScopeLock ScopeLock(&CriticalSection);
+	FScopeLock ScopeLock(&InstancesSection);
 	Instances.Add(this);
 }
 
 UE::FStallDetectorStats::~FStallDetectorStats()
 {
 	// Remove at the beginning of destruction
-	FScopeLock ScopeLock(&CriticalSection);
+	FScopeLock ScopeLock(&InstancesSection);
 	Instances.Remove(this);
 }
 
+void UE::FStallDetectorStats::OnStallCompleted(double InOverageSeconds)
+{
+	// we sync access around these for coherency reasons, can be polled from another thread (tabulation)
+	FScopeLock StatsLock(&StatsSection);
+	TriggerCount++;
+	OverageSeconds += InOverageSeconds;
+}
+
+void UE::FStallDetectorStats::TabulateStats(TArray<TabulatedResult>& TabulatedResults)
+{
+	TabulatedResults.Empty();
+
+	struct SortableStallStats
+	{
+		explicit SortableStallStats(const UE::FStallDetectorStats* InStats)
+			: StallStats(InStats)
+			, OverageRatio(0.0)
+		{
+			FScopeLock Lock(&InStats->StatsSection);
+			if (InStats->TriggerCount && InStats->BudgetSeconds > 0.0)
+			{
+				OverageRatio = (InStats->OverageSeconds / InStats->TriggerCount) / InStats->BudgetSeconds;
+			}
+		}
+
+		bool operator<(const SortableStallStats& InRhs) const
+		{
+			// NOTE THIS IS REVERSED TO PUT THE MOST OVERAGE AT THE FRONT
+			return OverageRatio > InRhs.OverageRatio;
+		}
+
+		const UE::FStallDetectorStats* StallStats;
+		double OverageRatio;
+	};
+
+	TArray<SortableStallStats> StatsArray;
+	FScopeLock InstancesLock(&UE::FStallDetectorStats::GetInstancesSection());
+	for (const UE::FStallDetectorStats* StallStats : UE::FStallDetectorStats::GetInstances())
+	{
+		if (StallStats->TriggerCount && StallStats->ReportingMode != UE::EStallDetectorReportingMode::Disabled)
+		{
+			StatsArray.Emplace(StallStats);
+		}
+	}
+
+	if (!StatsArray.IsEmpty())
+	{
+		StatsArray.Sort();
+
+		for (const SortableStallStats& Stat : StatsArray)
+		{
+			TabulatedResults.Emplace(TabulatedResult());
+			TabulatedResult& Result(TabulatedResults.Last());
+			Result.Stats = Stat.StallStats;
+
+			// we sync access around these for coherency reasons, can be polled from another thread (detector or scope)
+			FScopeLock Lock(&Result.Stats->StatsSection);
+			Result.TriggerCount = Result.Stats->TriggerCount;
+			Result.OverageSeconds = Result.Stats->OverageSeconds;
+		}
+	}
+}
+
+
 /**
 * Stall Detector
 **/
 
-FCriticalSection UE::FStallDetector::CriticalSection;
+FCriticalSection UE::FStallDetector::InstancesSection;
 TSet<UE::FStallDetector*> UE::FStallDetector::Instances;
 
 UE::FStallDetector::FStallDetector(FStallDetectorStats& InStats)
@@ -167,7 +232,7 @@ UE::FStallDetector::FStallDetector(FStallDetectorStats& InStats)
 	StartSeconds = FStallDetector::Seconds();
 
 	// Add at the end of construction
-	FScopeLock ScopeLock(&CriticalSection);
+	FScopeLock ScopeLock(&InstancesSection);
 	Instances.Add(this);
 }
 
@@ -175,7 +240,7 @@ UE::FStallDetector::~FStallDetector()
 {
 	// Remove at the beginning of destruction
 	{
-		FScopeLock ScopeLock(&CriticalSection);
+		FScopeLock ScopeLock(&InstancesSection);
 		Instances.Remove(this);
 	}
 
@@ -200,7 +265,7 @@ void UE::FStallDetector::Check(bool bIsComplete, double InWhenToCheckSeconds)
 	{
 		if (bIsComplete)
 		{
-			Stats.OverageSeconds += OverageSeconds;
+			Stats.OnStallCompleted(OverageSeconds);
 
 #if STALL_DETECTOR_DEBUG
 			FString OverageString = FString::Printf(TEXT("[FStallDetector] [%s] Overage of %f\n"), Stats.Name, OverageSeconds);
@@ -220,7 +285,6 @@ void UE::FStallDetector::Check(bool bIsComplete, double InWhenToCheckSeconds)
 				FString OverageString = FString::Printf(TEXT("[FStallDetector] [%s] Triggered at %f\n"), Stats.Name, CheckSeconds);
 				FPlatformMisc::LocalPrint(OverageString.GetCharArray().GetData());
 #endif
-				Stats.TriggerCount++;
 				OnStallDetected(ThreadId, DeltaSeconds);
 			}
 		}
@@ -260,15 +324,19 @@ void UE::FStallDetector::OnStallDetected(uint32 InThreadId, const double InElaps
 	EStallDetectorReportingMode ReportingMode = Stats.ReportingMode;
 
 	bool bDisableReporting = false;
+
 #if UE_BUILD_DEBUG
 	bDisableReporting |= true; // Do not generate a report in debug configurations due to performance characteristics
 #endif
+
+#if !STALL_DETECTOR_DEBUG
 	bDisableReporting |= FPlatformMisc::IsDebuggerPresent(); // Do not generate a report if we detect the debugger mucking with things
+#endif
 
 	if (bDisableReporting)
 	{
 #if !STALL_DETECTOR_DEBUG
-		ReportingMode = EStallDetectorReportingMode::Never;
+		ReportingMode = EStallDetectorReportingMode::Disabled;
 #endif
 	}
 
@@ -280,7 +348,7 @@ void UE::FStallDetector::OnStallDetected(uint32 InThreadId, const double InElaps
 	switch (ReportingMode)
 	{
 	case EStallDetectorReportingMode::First:
-		bSendReport = Stats.TriggerCount == 1;
+		bSendReport = !Stats.bReported;
 		break;
 
 	case EStallDetectorReportingMode::Always:
@@ -297,6 +365,7 @@ void UE::FStallDetector::OnStallDetected(uint32 InThreadId, const double InElaps
 
 	if (bSendReport)
 	{
+		Stats.bReported = true;
 		Stats.TotalReportedCount++;
 		const int NumStackFramesToIgnore = FPlatformTLS::GetCurrentThreadId() == InThreadId ? 2 : 0;
 		ReportStall(Stats.Name, InThreadId, NumStackFramesToIgnore);
diff --git a/Engine/Source/Runtime/Core/Public/ProfilingDebugging/StallDetector.h b/Engine/Source/Runtime/Core/Public/ProfilingDebugging/StallDetector.h
--- a/Engine/Source/Runtime/Core/Public/ProfilingDebugging/StallDetector.h
+++ b/Engine/Source/Runtime/Core/Public/ProfilingDebugging/StallDetector.h
@@ -36,7 +36,7 @@ namespace UE
 	**/
 	enum EStallDetectorReportingMode
 	{
-		Never,
+		Disabled,
 		First,
 		Always
 	};
@@ -57,12 +57,17 @@ namespace UE
 		*/
 		~FStallDetectorStats();
 
+		/**
+		* Stall dettector scope has ended, record the interval
+		*/
+		void OnStallCompleted(double InOverageSeconds);
+
 		/**
 		* Retrieve the critical section (lock) for accessing the instances list
 		*/
-		static FCriticalSection& GetCriticalSection()
+		static FCriticalSection& GetInstancesSection()
 		{
-			return CriticalSection;
+			return InstancesSection;
 		}
 
 		/**
@@ -73,6 +78,23 @@ namespace UE
 			return Instances;
 		}
 
+		/**
+		* 
+		*/
+		struct TabulatedResult
+		{
+			const UE::FStallDetectorStats* Stats;
+			uint32 TriggerCount;
+			double OverageSeconds;
+			
+			TabulatedResult()
+				: Stats(nullptr)
+				, TriggerCount(0)
+				, OverageSeconds(0.0)
+			{}
+		};
+		static void TabulateStats(TArray<TabulatedResult>& TabulatedResults);
+
 		// The name to refer to this callsite in the codebase
 		const TCHAR* Name;
 
@@ -82,11 +104,8 @@ namespace UE
 		// Drives the behavior about what to do when a budget is violated (triggered)
 		const EStallDetectorReportingMode ReportingMode;
 
-		// The number of times this callsite has been triggered
-		uint32 TriggerCount;
-
-		// The cumulative overage time for this callsite
-		double OverageSeconds;
+		// Has this been reported yet, this is different than the trigger count for coherency reasons
+		bool bReported;
 
 		// The total number of times all callsites have been triggered
 		static std::atomic<uint32> TotalTriggeredCount;
@@ -95,8 +114,17 @@ namespace UE
 		static std::atomic<uint32> TotalReportedCount;
 
 	private:
+		// The number of times this callsite has been triggered
+		uint32 TriggerCount;
+
+		// The cumulative overage time for this callsite
+		double OverageSeconds;
+
+		// Guards access to the stats from multiple threads, for coherency
+		mutable FCriticalSection StatsSection;
+
 		// Guards access to the instances container from multiple threads
-		static FCriticalSection CriticalSection;
+		static FCriticalSection InstancesSection;
 
 		// The pointers to every constructed instance of this class
 		static TSet<FStallDetectorStats*> Instances;
@@ -140,9 +168,9 @@ namespace UE
 		/**
 		* Retrieve the critical section (lock) for accessing the instances list
 		*/
-		static FCriticalSection& GetCriticalSection()
+		static FCriticalSection& GetInstancesSection()
 		{
-			return CriticalSection;
+			return InstancesSection;
 		}
 
 		/**
@@ -183,24 +211,40 @@ namespace UE
 		void OnStallDetected(uint32 InThreadId, const double InElapsedSeconds);
 
 		// Guards access to the instances container from multiple threads
-		static FCriticalSection CriticalSection;
+		static FCriticalSection InstancesSection;
 
 		// The pointers to every constructed instance of this class
 		static TSet<FStallDetector*> Instances;
 	};
 }
 
+/**
+/* Counts, but does not report, a stall.
+**/
+#define SCOPE_STALL_COUNTER(InName, InBudgetSeconds) \
+static UE::FStallDetectorStats PREPROCESSOR_JOIN(StallDetectorStats, __LINE__) (TEXT(#InName), InBudgetSeconds, UE::EStallDetectorReportingMode::Disabled); \
+UE::FStallDetector PREPROCESSOR_JOIN(ScopeStallDetector, __LINE__)(PREPROCESSOR_JOIN(StallDetectorStats, __LINE__));
+
+/**
+/* Counts and sends a report of the first occurence of a stall.
+**/
 #define SCOPE_STALL_REPORTER(InName, InBudgetSeconds) \
 static UE::FStallDetectorStats PREPROCESSOR_JOIN(StallDetectorStats, __LINE__) (TEXT(#InName), InBudgetSeconds, UE::EStallDetectorReportingMode::First); \
 UE::FStallDetector PREPROCESSOR_JOIN(ScopeStallDetector, __LINE__)(PREPROCESSOR_JOIN(StallDetectorStats, __LINE__));
 
-#define SCOPE_STALL_COUNTER(InName, InBudgetSeconds) \
-static UE::FStallDetectorStats PREPROCESSOR_JOIN(StallDetectorStats, __LINE__) (TEXT(#InName), InBudgetSeconds, UE::EStallDetectorReportingMode::Never); \
+/**
+* Counts and reports every occurence of a stall.
+* Use with extreme caution as you could overload the report backend!
+* This mainly exists to help test and exercise the stall detector system.
+**/
+#define SCOPE_STALL_REPORTER_ALWAYS(InName, InBudgetSeconds) \
+static UE::FStallDetectorStats PREPROCESSOR_JOIN(StallDetectorStats, __LINE__) (TEXT(#InName), InBudgetSeconds, UE::EStallDetectorReportingMode::Always); \
 UE::FStallDetector PREPROCESSOR_JOIN(ScopeStallDetector, __LINE__)(PREPROCESSOR_JOIN(StallDetectorStats, __LINE__));
 
 #else // STALL_DETECTOR
 
-#define SCOPE_STALL_REPORTER(InName, InBudgetSeconds)
 #define SCOPE_STALL_COUNTER(InName, InBudgetSeconds)
+#define SCOPE_STALL_REPORTER(InName, InBudgetSeconds)
+#define SCOPE_STALL_REPORTER_ALWAYS(InName, InBudgetSeconds)
 
 #endif // STALL_DETECTOR
\ No newline at end of file
diff --git a/Engine/Source/Runtime/Engine/Private/Analytics/EditorSessionSummaryWriter.cpp b/Engine/Source/Runtime/Engine/Private/Analytics/EditorSessionSummaryWriter.cpp
--- a/Engine/Source/Runtime/Engine/Private/Analytics/EditorSessionSummaryWriter.cpp
+++ b/Engine/Source/Runtime/Engine/Private/Analytics/EditorSessionSummaryWriter.cpp
@@ -27,6 +27,7 @@
 #include "Misc/EngineVersion.h"
 #include "Misc/Guid.h"
 #include "GenericPlatform/GenericPlatformCrashContext.h"
+#include "ProfilingDebugging/StallDetector.h"
 
 #include "EditorAnalyticsSession.h"
 #include "IVREditorModule.h"
@@ -341,6 +342,23 @@ void FEditorSessionSummaryWriter::Tick(float DeltaTime)
 		NextDebuggerCheckSecs = CurrentTimeSecs + EditorSessionWriterDefs::DebuggerCheckPeriodSeconds;
 	}
 
+	// Accumulate stall stats
+#if STALL_DETECTOR
+	TArray<UE::FStallDetectorStats::TabulatedResult> StallResults;
+	UE::FStallDetectorStats::TabulateStats(StallResults);
+	if (!StallResults.IsEmpty())
+	{
+		UE::FStallDetectorStats::TabulatedResult TopResult(StallResults[0]);
+		CurrentSession->TopStallName = TopResult.Stats->Name;
+		CurrentSession->TopStallBudgetSeconds = TopResult.Stats->BudgetSeconds;
+		CurrentSession->TopStallOverageSeconds = TopResult.OverageSeconds;
+		CurrentSession->TopStallTriggerCount = TopResult.TriggerCount;
+	}
+
+	CurrentSession->TotalStallCount = UE::FStallDetectorStats::TotalTriggeredCount;
+	CurrentSession->TotalStallReported = UE::FStallDetectorStats::TotalReportedCount;
+#endif // STALL_DETECTOR
+
 	if (bSaveSession)
 	{
 		// Saves asynchronously, saving up to 5ms in the game thread. Saving also updates session duration/timestamp/userIdle/editorIdle
