Add a "Process Diagnostics" flags word that enumerates the presence of various performance-impacting engine and platform diagnostics.

On Windows this checks for the kernel's various memory sanitizers (Set by gflags).

diff --git a/Engine/Source/Developer/EditorAnalyticsSession/Private/EditorAnalyticsSession.cpp b/Engine/Source/Developer/EditorAnalyticsSession/Private/EditorAnalyticsSession.cpp
--- a/Engine/Source/Developer/EditorAnalyticsSession/Private/EditorAnalyticsSession.cpp
+++ b/Engine/Source/Developer/EditorAnalyticsSession/Private/EditorAnalyticsSession.cpp
@@ -122,6 +122,8 @@ namespace EditorAnalyticsDefs
 	static const FString IsInVRModeStoreKey(TEXT("IsInVRMode"));
 	static const FString IsCrcExeMissingStoreKey(TEXT("IsCrcExeMissing"));
 	static const FString IsLowDriveSpaceStoreKey(TEXT("IsLowDriveSpace"));
+
+	static const FString ProcessDiagnosticsStoreKey(TEXT("ProcessDiagnostics"));
 }
 
 // Utilities for writing to stored values
@@ -477,6 +479,8 @@ namespace EditorAnalyticsUtils
 		Session.bIsLowDriveSpace = EditorAnalyticsUtils::GetStoredBool(SectionName, EditorAnalyticsDefs::IsLowDriveSpaceStoreKey);
 		Session.bIsCrcExeMissing = EditorAnalyticsUtils::GetStoredBool(SectionName, EditorAnalyticsDefs::IsCrcExeMissingStoreKey);
 
+		GET_STORED_INT(ProcessDiagnostics);
+
 		// Analyze the logged events and update corresponding fields in the session.
 		UpdateSessionFromLogAnalysis(Session);
 	}
@@ -526,6 +530,8 @@ FEditorAnalyticsSession::FEditorAnalyticsSession()
 	bIsLowDriveSpace = false;
 	bIsCrcExeMissing = false;
 	bIsDebuggerIgnored = false;
+
+	ProcessDiagnostics = 0x0;
 }
 
 bool FEditorAnalyticsSession::Lock(FTimespan Timeout)
@@ -653,6 +659,7 @@ bool FEditorAnalyticsSession::Save()
 			{EditorAnalyticsDefs::IsInEnterpriseStoreKey,   EditorAnalyticsUtils::BoolToStoredString(bIsInEnterprise)},
 			{EditorAnalyticsDefs::IsInVRModeStoreKey,       EditorAnalyticsUtils::BoolToStoredString(bIsInVRMode)},
 			{EditorAnalyticsDefs::IsLowDriveSpaceStoreKey,  EditorAnalyticsUtils::BoolToStoredString(bIsLowDriveSpace)},
+			{EditorAnalyticsDefs::ProcessDiagnosticsStoreKey, FString::FromInt(ProcessDiagnostics)},
 		};
 
 		if (ExitCode.IsSet())
@@ -776,6 +783,7 @@ bool FEditorAnalyticsSession::Delete() const
 	FPlatformMisc::DeleteStoredValue(EditorAnalyticsDefs::StoreId, SectionName, EditorAnalyticsDefs::IsInVRModeStoreKey);
 	FPlatformMisc::DeleteStoredValue(EditorAnalyticsDefs::StoreId, SectionName, EditorAnalyticsDefs::IsLowDriveSpaceStoreKey);
 	FPlatformMisc::DeleteStoredValue(EditorAnalyticsDefs::StoreId, SectionName, EditorAnalyticsDefs::IsCrcExeMissingStoreKey);
+	FPlatformMisc::DeleteStoredValue(EditorAnalyticsDefs::StoreId, SectionName, EditorAnalyticsDefs::ProcessDiagnosticsStoreKey);
 
 	// Delete the log files.
 	EditorAnalyticsUtils::DeleteLogEvents(*this);
diff --git a/Engine/Source/Developer/EditorAnalyticsSession/Private/EditorSessionSummarySender.cpp b/Engine/Source/Developer/EditorAnalyticsSession/Private/EditorSessionSummarySender.cpp
--- a/Engine/Source/Developer/EditorAnalyticsSession/Private/EditorSessionSummarySender.cpp
+++ b/Engine/Source/Developer/EditorAnalyticsSession/Private/EditorSessionSummarySender.cpp
@@ -262,6 +262,7 @@ void FEditorSessionSummarySender::SendSessionSummaryEvent(const FEditorAnalyticR
 	AnalyticsAttributes.Emplace(TEXT("SentFrom"), Sender);
 	AnalyticsAttributes.Emplace(TEXT("MonitorPid"), Session.MonitorProcessID); // For out-of-process monitoring, if this is 0, this will mean CRC failed to launch or crashed very early.
 	AnalyticsAttributes.Emplace(TEXT("DelayedCrashCount"), ReportData.DelayedCrashCount); // Piggybacked information from 'minimal crash sessions', previous Editor instances crashes recorded before analytics was initialized.
+	AnalyticsAttributes.Emplace(TEXT("ProcessDiagnostics"), Session.ProcessDiagnostics);
 
 	bool bShouldAttachMonitorLog = (ShutdownTypeString != EditorSessionSenderDefs::ShutdownSessionToken && ShutdownTypeString != EditorSessionSenderDefs::TerminatedSessionToken);
 
diff --git a/Engine/Source/Developer/EditorAnalyticsSession/Public/EditorAnalyticsSession.h b/Engine/Source/Developer/EditorAnalyticsSession/Public/EditorAnalyticsSession.h
--- a/Engine/Source/Developer/EditorAnalyticsSession/Public/EditorAnalyticsSession.h
+++ b/Engine/Source/Developer/EditorAnalyticsSession/Public/EditorAnalyticsSession.h
@@ -98,6 +98,8 @@ struct EDITORANALYTICSSESSION_API FEditorAnalyticsSession
 	bool bIsCrcExeMissing: 1; // CrashReportClient executable is missing? To explain with MonitorProcessID would be zero.
 	bool bIsDebuggerIgnored: 1; // True if GIgnoreDebugger is true.
 
+	uint32 ProcessDiagnostics;
+
 	FEditorAnalyticsSession();
 
 	/** 
diff --git a/Engine/Source/Runtime/Core/Private/GenericPlatform/GenericPlatformMisc.cpp b/Engine/Source/Runtime/Core/Private/GenericPlatform/GenericPlatformMisc.cpp
--- a/Engine/Source/Runtime/Core/Private/GenericPlatform/GenericPlatformMisc.cpp
+++ b/Engine/Source/Runtime/Core/Private/GenericPlatform/GenericPlatformMisc.cpp
@@ -369,6 +369,28 @@ void FGenericPlatformMisc::SubmitErrorReport( const TCHAR* InErrorHist, EErrorRe
 	}
 }
 
+EProcessDiagnosticFlags FGenericPlatformMisc::GetProcessDiagnostics()
+{
+	static EProcessDiagnosticFlags FoundDiagnostics = []() -> EProcessDiagnosticFlags
+	{
+		EProcessDiagnosticFlags Result = EProcessDiagnosticFlags::None;
+		const TCHAR* CommandLine = FCommandLine::Get();
+
+		if (FCString::Stristr(CommandLine, TEXT("-ansimalloc")))
+		{
+			Result |= EProcessDiagnosticFlags::AnsiMalloc;
+		}
+
+		if (FCString::Stristr(CommandLine, TEXT("-stompmalloc")))
+		{
+			Result |= EProcessDiagnosticFlags::StompMalloc;
+		}
+
+		return Result;
+	}();
+
+	return FoundDiagnostics;
+}
 
 FString FGenericPlatformMisc::GetCPUVendor()
 {
diff --git a/Engine/Source/Runtime/Core/Private/Unix/UnixPlatformMemory.cpp b/Engine/Source/Runtime/Core/Private/Unix/UnixPlatformMemory.cpp
--- a/Engine/Source/Runtime/Core/Private/Unix/UnixPlatformMemory.cpp
+++ b/Engine/Source/Runtime/Core/Private/Unix/UnixPlatformMemory.cpp
@@ -158,6 +158,7 @@ class FMalloc* FUnixPlatformMemory::BaseAllocator()
 #endif // PLATFORM_SUPPORTS_JEMALLOC
 				if (FCStringAnsi::Stricmp(Arg, "-ansimalloc") == 0)
 				{
+					// see FPlatformMisc::GetProcessDiagnostics()
 					AllocatorToUse = EMemoryAllocatorToUse::Ansi;
 					break;
 				}
@@ -224,6 +225,7 @@ class FMalloc* FUnixPlatformMemory::BaseAllocator()
 #if WITH_MALLOC_STOMP
 				if (FCStringAnsi::Stricmp(Arg, "-stompmalloc") == 0)
 				{
+					// see FPlatformMisc::GetProcessDiagnostics()
 					AllocatorToUse = EMemoryAllocatorToUse::Stomp;
 					break;
 				}
diff --git a/Engine/Source/Runtime/Core/Private/Windows/WindowsPlatformMemory.cpp b/Engine/Source/Runtime/Core/Private/Windows/WindowsPlatformMemory.cpp
--- a/Engine/Source/Runtime/Core/Private/Windows/WindowsPlatformMemory.cpp
+++ b/Engine/Source/Runtime/Core/Private/Windows/WindowsPlatformMemory.cpp
@@ -125,6 +125,7 @@ FMalloc* FWindowsPlatformMemory::BaseAllocator()
 
 	if (FCString::Stristr(CommandLine, TEXT("-ansimalloc")))
 	{
+		// see FPlatformMisc::GetProcessDiagnostics()
 		AllocatorToUse = EMemoryAllocatorToUse::Ansi;
 	}
 #if TBB_ALLOCATOR_ALLOWED
@@ -156,6 +157,7 @@ FMalloc* FWindowsPlatformMemory::BaseAllocator()
 #if WITH_MALLOC_STOMP
 	else if (FCString::Stristr(CommandLine, TEXT("-stompmalloc")))
 	{
+		// see FPlatformMisc::GetProcessDiagnostics()
 		AllocatorToUse = EMemoryAllocatorToUse::Stomp;
 	}
 #endif // WITH_MALLOC_STOMP
diff --git a/Engine/Source/Runtime/Core/Private/Windows/WindowsPlatformMisc.cpp b/Engine/Source/Runtime/Core/Private/Windows/WindowsPlatformMisc.cpp
--- a/Engine/Source/Runtime/Core/Private/Windows/WindowsPlatformMisc.cpp
+++ b/Engine/Source/Runtime/Core/Private/Windows/WindowsPlatformMisc.cpp
@@ -52,6 +52,7 @@ THIRD_PARTY_INCLUDES_START
 	#include <ShlObj.h>
 	#include <IntShCut.h>
 	#include <shellapi.h>
+	#include <shlwapi.h>
 	#include <IPHlpApi.h>
 	#include <VersionHelpers.h>
 THIRD_PARTY_INCLUDES_END
@@ -840,6 +841,41 @@ bool FWindowsPlatformMisc::IsDebuggerPresent()
 {
 	return !GIgnoreDebugger && !!::IsDebuggerPresent();
 }
+
+EProcessDiagnosticFlags FWindowsPlatformMisc::GetProcessDiagnostics()
+{
+	static EProcessDiagnosticFlags FoundDiagnostics = []() -> EProcessDiagnosticFlags
+	{
+		EProcessDiagnosticFlags Result = FGenericPlatformMisc::GetProcessDiagnostics();
+
+		TCHAR* ImageFileName = PathFindFileName(FPlatformProcess::ExecutablePath());
+		FString ImageFileSubkey = FString::Printf(TEXT("SOFTWARE\\Microsoft\\Windows NT\\CurrentVersion\\Image File Execution Options\\%s"), ImageFileName);
+
+		// via https://docs.microsoft.com/en-us/windows-hardware/drivers/debugger/gflags-flag-table
+		constexpr uint32 MemorySanitizerMask =
+			0x00000010| // FLG_HEAP_ENABLE_TAIL_CHECK
+			0x00000020| // FLG_HEAP_ENABLE_FREE_CHECK
+			0x00000080| // FLG_HEAP_VALIDATE_ALL
+			0x00000100| // FLG_APPLICATION_VERIFIER
+			0x00000800| // FLG_HEAP_ENABLE_TAGGING
+			0x00008000| // FLG_HEAP_ENABLE_TAG_BY_DLL
+			0x00200000| // FLG_HEAP_DISABLE_COALESCING
+			0x02000000; // FLG_HEAP_PAGE_ALLOCS
+
+		DWORD Data, DataCount = sizeof(Data);
+		if (ERROR_SUCCESS == RegGetValue(HKEY_LOCAL_MACHINE, *ImageFileSubkey, TEXT("GlobalFlag"), RRF_RT_REG_DWORD, nullptr, &Data, &DataCount))
+		{
+			if (MemorySanitizerMask & Data)
+			{
+				Result |= EProcessDiagnosticFlags::MemorySanitizer;
+			}
+		}
+
+		return Result;
+	}();
+
+	return FoundDiagnostics;
+}
 #endif //!UE_BUILD_SHIPPING
 
 #if STATS || ENABLE_STATNAMEDEVENTS
diff --git a/Engine/Source/Runtime/Core/Public/GenericPlatform/GenericPlatformMisc.h b/Engine/Source/Runtime/Core/Public/GenericPlatform/GenericPlatformMisc.h
--- a/Engine/Source/Runtime/Core/Public/GenericPlatform/GenericPlatformMisc.h
+++ b/Engine/Source/Runtime/Core/Public/GenericPlatform/GenericPlatformMisc.h
@@ -6,6 +6,7 @@
 #include "CoreFwd.h"
 #include "HAL/PlatformCrt.h"
 #include "Misc/CompressionFlags.h"
+#include "Misc/EnumClassFlags.h"
 #include "Math/NumericLimits.h"
 
 class Error;
@@ -325,6 +326,16 @@ enum class ENetworkConnectionType : uint8
 	Ethernet,
 };
 
+enum class EProcessDiagnosticFlags : uint32
+{
+	None            = 0,
+	MemorySanitizer = 1 << 0,
+	AnsiMalloc      = 1 << 1,
+	StompMalloc     = 1 << 2,
+};
+
+ENUM_CLASS_FLAGS(EProcessDiagnosticFlags);
+
 /**
  * Returns the string representation of the specified ENetworkConnection value.
  *
@@ -474,12 +485,15 @@ struct CORE_API FGenericPlatformMisc
 	FORCEINLINE static bool IsDebuggerPresent()
 	{
 #if UE_BUILD_SHIPPING
-		return 0;
+		return false;
 #else
-		return 1; // unknown platforms return true so that they can crash into a debugger
+		return true; // unknown platforms return true so that they can crash into a debugger
 #endif
 	}
 
+	/** Returns non-zero if a performance-sensitive diagnostic, like a memory sanitizer, is detected for this process */
+	static EProcessDiagnosticFlags GetProcessDiagnostics();
+
 	/** Break into the debugger, if IsDebuggerPresent returns true, otherwise do nothing  */
 	UE_DEPRECATED(4.19, "FPlatformMisc::DebugBreak is deprecated. Use the UE_DEBUG_BREAK() macro instead.")
 	FORCEINLINE static void DebugBreak()
diff --git a/Engine/Source/Runtime/Core/Public/Windows/WindowsPlatformMisc.h b/Engine/Source/Runtime/Core/Public/Windows/WindowsPlatformMisc.h
--- a/Engine/Source/Runtime/Core/Public/Windows/WindowsPlatformMisc.h
+++ b/Engine/Source/Runtime/Core/Public/Windows/WindowsPlatformMisc.h
@@ -51,6 +51,7 @@ struct CORE_API FWindowsPlatformMisc
 
 #if !UE_BUILD_SHIPPING
 	static bool IsDebuggerPresent();
+	static EProcessDiagnosticFlags GetProcessDiagnostics();
 #endif
 
 #if STATS || ENABLE_STATNAMEDEVENTS
diff --git a/Engine/Source/Runtime/Engine/Private/Analytics/EditorSessionSummaryWriter.cpp b/Engine/Source/Runtime/Engine/Private/Analytics/EditorSessionSummaryWriter.cpp
--- a/Engine/Source/Runtime/Engine/Private/Analytics/EditorSessionSummaryWriter.cpp
+++ b/Engine/Source/Runtime/Engine/Private/Analytics/EditorSessionSummaryWriter.cpp
@@ -339,6 +339,7 @@ void FEditorSessionSummaryWriter::Tick(float DeltaTime)
 			}
 			bSaveSession = true;
 		}
+
 		NextDebuggerCheckSecs = CurrentTimeSecs + EditorSessionWriterDefs::DebuggerCheckPeriodSeconds;
 	}
 
@@ -466,6 +467,7 @@ TUniquePtr<FEditorAnalyticsSession> FEditorSessionSummaryWriter::CreateCurrentSe
 	Session->Timestamp = FDateTime::UtcNow();
 	Session->bIsDebugger = FPlatformMisc::IsDebuggerPresent();
 	Session->bWasEverDebugger = FPlatformMisc::IsDebuggerPresent();
+	Session->ProcessDiagnostics = static_cast<uint32>(FPlatformMisc::GetProcessDiagnostics());
 	Session->CurrentUserActivity = GetUserActivityString();
 	Session->bIsVanilla = GEngine && GEngine->IsVanillaProduct();
 	Session->CommandLine = FCommandLine::GetForLogging();
