Bumped the EditorAnalyticSession version number following a format change.
Renamed IsDebbuger analytic field as IsDebuggerPresent.

diff --git a/Engine/Source/Developer/EditorAnalyticsSession/Private/EditorAnalyticsSession.cpp b/Engine/Source/Developer/EditorAnalyticsSession/Private/EditorAnalyticsSession.cpp
--- a/Engine/Source/Developer/EditorAnalyticsSession/Private/EditorAnalyticsSession.cpp
+++ b/Engine/Source/Developer/EditorAnalyticsSession/Private/EditorAnalyticsSession.cpp
@@ -30,14 +30,18 @@ namespace EditorAnalyticsDefs
 	//   Version 1_3 : Added SessionTickCount, UserInteractionCount, IsCrcExeMissing, IsUserLoggingOut, MonitorExitCode and readded lost code to save/load/delete IsLowDriveSpace for 4.26.0.
 	//   Version 1_4 : Added CommandLine, EngineTickCount, LastTickTimestamp, DeathTimestamp and IsDebuggerIgnored for 4.26.0.
 	//   Version 1_5 : Added Stall Detector stats for 5.0.
+	//   Version 1_6 : Added ProcessDiagnostics, Renamed IsDebugger as IsDebuggerPresent
 	static const FString StoreId(TEXT("Epic Games"));
 	static const FString SessionSummaryRoot(TEXT("Unreal Engine/Session Summary"));
-	static const FString SessionSummarySection_1_0 = SessionSummaryRoot / TEXT("1_0"); // The session format used by older versions.
-	static const FString SessionSummarySection_1_1 = SessionSummaryRoot / TEXT("1_1");
-	static const FString SessionSummarySection_1_2 = SessionSummaryRoot / TEXT("1_2");
-	static const FString SessionSummarySection_1_3 = SessionSummaryRoot / TEXT("1_3");
-	static const FString SessionSummarySection_1_4 = SessionSummaryRoot / TEXT("1_4");
-	static const FString SessionSummarySection = SessionSummaryRoot / TEXT("1_5"); // The current session format.
+	static const FString DeprecatedVersions[] = { // The session format used by older versions.
+		SessionSummaryRoot / TEXT("1_0"),
+		SessionSummaryRoot / TEXT("1_1"),
+		SessionSummaryRoot / TEXT("1_2"),
+		SessionSummaryRoot / TEXT("1_3"),
+		SessionSummaryRoot / TEXT("1_4"),
+		SessionSummaryRoot / TEXT("1_5"),
+	};
+	static const FString SessionSummarySection = SessionSummaryRoot / TEXT("1_6"); // The current session format.
 	static const FString GlobalLockName(TEXT("UE4_SessionSummary_Lock"));
 	static const FString SessionListStoreKey(TEXT("SessionList"));
 
@@ -110,7 +114,7 @@ namespace EditorAnalyticsDefs
 	// boolean flags
 	static const FString IsCrashStoreKey(TEXT("IsCrash"));
 	static const FString IsGPUCrashStoreKey(TEXT("IsGPUCrash"));
-	static const FString IsDebuggerStoreKey(TEXT("IsDebugger"));
+	static const FString IsDebuggerPresentStoreKey(TEXT("IsDebuggerPresent"));
 	static const FString IsDebuggerIgnoredStoreKey(TEXT("IsDebuggerIgnored"));
 	static const FString WasDebuggerStoreKey(TEXT("WasEverDebugger"));
 	static const FString IsVanillaStoreKey(TEXT("IsVanilla"));
@@ -466,7 +470,7 @@ namespace EditorAnalyticsUtils
 		Session.bIs64BitOS = EditorAnalyticsUtils::GetStoredBool(SectionName, EditorAnalyticsDefs::bIs64BitOSStoreKey);
 		Session.bCrashed = EditorAnalyticsUtils::GetStoredBool(SectionName, EditorAnalyticsDefs::IsCrashStoreKey);
 		Session.bGPUCrashed = EditorAnalyticsUtils::GetStoredBool(SectionName, EditorAnalyticsDefs::IsGPUCrashStoreKey);
-		Session.bIsDebugger = EditorAnalyticsUtils::GetStoredBool(SectionName, EditorAnalyticsDefs::IsDebuggerStoreKey);
+		Session.bIsDebuggerPresent = EditorAnalyticsUtils::GetStoredBool(SectionName, EditorAnalyticsDefs::IsDebuggerPresentStoreKey);
 		Session.bIsDebuggerIgnored = EditorAnalyticsUtils::GetStoredBool(SectionName, EditorAnalyticsDefs::IsDebuggerIgnoredStoreKey);
 		Session.bWasEverDebugger = EditorAnalyticsUtils::GetStoredBool(SectionName, EditorAnalyticsDefs::WasDebuggerStoreKey);
 		Session.bIsVanilla = EditorAnalyticsUtils::GetStoredBool(SectionName, EditorAnalyticsDefs::IsVanillaStoreKey);
@@ -517,7 +521,7 @@ FEditorAnalyticsSession::FEditorAnalyticsSession()
 	bIs64BitOS = false;
 	bCrashed = false;
 	bGPUCrashed = false;
-	bIsDebugger = false;
+	bIsDebuggerPresent = false;
 	bWasEverDebugger = false;
 	bIsVanilla = false;
 	bIsTerminating = false;
@@ -649,7 +653,7 @@ bool FEditorAnalyticsSession::Save()
 			{EditorAnalyticsDefs::TopStallBudgetSecondsStoreKey,  FString::SanitizeFloat(TopStallBudgetSeconds)},
 			{EditorAnalyticsDefs::TopStallOverageSecondsStoreKey, FString::SanitizeFloat(TopStallOverageSeconds)},
 			{EditorAnalyticsDefs::TopStallTriggerCountStoreKey,   FString::FromInt(TopStallTriggerCount)},
-			{EditorAnalyticsDefs::IsDebuggerStoreKey,       EditorAnalyticsUtils::BoolToStoredString(bIsDebugger)},
+			{EditorAnalyticsDefs::IsDebuggerPresentStoreKey,EditorAnalyticsUtils::BoolToStoredString(bIsDebuggerPresent)},
 			{EditorAnalyticsDefs::IsDebuggerIgnoredStoreKey,EditorAnalyticsUtils::BoolToStoredString(bIsDebuggerIgnored)},
 			{EditorAnalyticsDefs::WasDebuggerStoreKey,      EditorAnalyticsUtils::BoolToStoredString(bWasEverDebugger)},
 			{EditorAnalyticsDefs::IsVanillaStoreKey,        EditorAnalyticsUtils::BoolToStoredString(bIsVanilla)},
@@ -771,7 +775,7 @@ bool FEditorAnalyticsSession::Delete() const
 
 	FPlatformMisc::DeleteStoredValue(EditorAnalyticsDefs::StoreId, SectionName, EditorAnalyticsDefs::IsCrashStoreKey);
 	FPlatformMisc::DeleteStoredValue(EditorAnalyticsDefs::StoreId, SectionName, EditorAnalyticsDefs::IsGPUCrashStoreKey);
-	FPlatformMisc::DeleteStoredValue(EditorAnalyticsDefs::StoreId, SectionName, EditorAnalyticsDefs::IsDebuggerStoreKey);
+	FPlatformMisc::DeleteStoredValue(EditorAnalyticsDefs::StoreId, SectionName, EditorAnalyticsDefs::IsDebuggerPresentStoreKey);
 	FPlatformMisc::DeleteStoredValue(EditorAnalyticsDefs::StoreId, SectionName, EditorAnalyticsDefs::IsDebuggerIgnoredStoreKey);
 	FPlatformMisc::DeleteStoredValue(EditorAnalyticsDefs::StoreId, SectionName, EditorAnalyticsDefs::WasDebuggerStoreKey);
 	FPlatformMisc::DeleteStoredValue(EditorAnalyticsDefs::StoreId, SectionName, EditorAnalyticsDefs::IsVanillaStoreKey);
@@ -853,7 +857,7 @@ void FEditorAnalyticsSession::CleanupOutdatedIncompatibleSessions(const FTimespa
 		return;
 	}
 
-	// Helper function to scan and clear sessions stored in sections corresponding to format version '1_0' and '1_1'.
+	// Helper function to scan and clear sessions stored in sections corresponding to older versions.
 	auto CleanupVersionedSection = [&MaxAge](const FString& SectionVersion)
 	{
 		// Try to retreive the session list corresponding the specified session format.
@@ -886,13 +890,12 @@ void FEditorAnalyticsSession::CleanupOutdatedIncompatibleSessions(const FTimespa
 		FPlatformMisc::DeleteStoredSection(EditorAnalyticsDefs::StoreId, SectionVersion);
 	};
 
-	// The current section format is 1_3. The older sections are considered incompatible and will be trimmed unless it contains a valid session young enough that would be picked up
+	// Delete older and incompatible sections unless it contains a valid session young enough that would be picked up
 	// if an older Editor with compatible format was launched again.
-	CleanupVersionedSection(EditorAnalyticsDefs::SessionSummarySection_1_0);
-	CleanupVersionedSection(EditorAnalyticsDefs::SessionSummarySection_1_1);
-	CleanupVersionedSection(EditorAnalyticsDefs::SessionSummarySection_1_2);
-	CleanupVersionedSection(EditorAnalyticsDefs::SessionSummarySection_1_3);
-	CleanupVersionedSection(EditorAnalyticsDefs::SessionSummarySection_1_4);
+	for (int i = 0; i < UE_ARRAY_COUNT(EditorAnalyticsDefs::DeprecatedVersions); ++i)
+	{
+		CleanupVersionedSection(EditorAnalyticsDefs::DeprecatedVersions[i]);
+	}
 }
 
 void FEditorAnalyticsSession::LogEvent(EEventType InEventType, const FDateTime& InTimestamp)
diff --git a/Engine/Source/Developer/EditorAnalyticsSession/Public/EditorAnalyticsSession.h b/Engine/Source/Developer/EditorAnalyticsSession/Public/EditorAnalyticsSession.h
--- a/Engine/Source/Developer/EditorAnalyticsSession/Public/EditorAnalyticsSession.h
+++ b/Engine/Source/Developer/EditorAnalyticsSession/Public/EditorAnalyticsSession.h
@@ -85,7 +85,7 @@ struct EDITORANALYTICSSESSION_API FEditorAnalyticsSession
 	bool bIs64BitOS : 1;
 	bool bCrashed : 1;
 	bool bGPUCrashed : 1;
-	bool bIsDebugger : 1;
+	bool bIsDebuggerPresent : 1;
 	bool bWasEverDebugger : 1;
 	bool bIsVanilla : 1;
 	bool bIsTerminating : 1;
diff --git a/Engine/Source/Runtime/Engine/Private/Analytics/EditorSessionSummaryWriter.cpp b/Engine/Source/Runtime/Engine/Private/Analytics/EditorSessionSummaryWriter.cpp
--- a/Engine/Source/Runtime/Engine/Private/Analytics/EditorSessionSummaryWriter.cpp
+++ b/Engine/Source/Runtime/Engine/Private/Analytics/EditorSessionSummaryWriter.cpp
@@ -330,9 +330,9 @@ void FEditorSessionSummaryWriter::Tick(float DeltaTime)
 
 		// Check if the debugger is present.
 		bool bIsDebuggerPresent = FPlatformMisc::IsDebuggerPresent();
-		if (CurrentSession->bIsDebugger != bIsDebuggerPresent)
+		if (CurrentSession->bIsDebuggerPresent != bIsDebuggerPresent)
 		{
-			CurrentSession->bIsDebugger = bIsDebuggerPresent;
+			CurrentSession->bIsDebuggerPresent = bIsDebuggerPresent;
 			if (bIsDebuggerPresent && !CurrentSession->bWasEverDebugger)
 			{
 				CurrentSession->bWasEverDebugger = true;
@@ -465,7 +465,7 @@ TUniquePtr<FEditorAnalyticsSession> FEditorSessionSummaryWriter::CreateCurrentSe
 	Session->StartupTimestamp = StartupTimeUtc;
 	Session->LastTickTimestamp = FDateTime::UtcNow();
 	Session->Timestamp = FDateTime::UtcNow();
-	Session->bIsDebugger = FPlatformMisc::IsDebuggerPresent();
+	Session->bIsDebuggerPresent = FPlatformMisc::IsDebuggerPresent();
 	Session->bWasEverDebugger = FPlatformMisc::IsDebuggerPresent();
 	Session->ProcessDiagnostics = static_cast<uint32>(FPlatformMisc::GetProcessDiagnostics());
 	Session->CurrentUserActivity = GetUserActivityString();
