Report all Editor bootstrapping failures captured by CrashReportClientEditor (when the Editor dies before analytics could be initialized).

diff --git a/Engine/Source/Developer/EditorAnalyticsSession/Private/EditorSessionSummarySender.cpp b/Engine/Source/Developer/EditorAnalyticsSession/Private/EditorSessionSummarySender.cpp
--- a/Engine/Source/Developer/EditorAnalyticsSession/Private/EditorSessionSummarySender.cpp
+++ b/Engine/Source/Developer/EditorAnalyticsSession/Private/EditorSessionSummarySender.cpp
@@ -35,8 +35,14 @@ struct FEditorAnalyticReportData
 	/** The analytic session to report. */
 	FEditorAnalyticsSession Session;
 
-	/** Number of Editor instances (same engine version than Session) that crashed prior the analytic was running and wasn't accounted for yet. */
+	/** Number of Editor instances (same engine version than Session) that crashed prior the analytic was initialized and wasn't reported yet. */
 	uint32 DelayedCrashCount = 0;
+
+	/** Number of Editor instances (same engine version than Session) that was killed by user prior the analytic was initialized and wasn't reported yet. */
+	uint32 DelayedKillCount = 0;
+
+	/* Number of Editor instance (same engine version than Session) that exited abnormally prior the analytic was initialized and wasn't reported yet. */
+	uint32 DelayedAbnormalCount = 0;
 };
 
 
@@ -123,7 +129,7 @@ void FEditorSessionSummarySender::SendStoredSessions(const bool bForceSendOwnedS
 						continue; // Don't schedule this minimal session for reporting/deletiong yet, it's data must be piggybacked off a full and compatible session to be reported/deleted.
 					}
 
-					SessionsToReport.Add(FEditorAnalyticReportData{Session,0});
+					SessionsToReport.Add(FEditorAnalyticReportData{Session,0, 0, 0});
 				}
 				SessionsToDelete.Add(Session);
 			}
@@ -135,7 +141,25 @@ void FEditorSessionSummarySender::SendStoredSessions(const bool bForceSendOwnedS
 				if (FEditorAnalyticReportData* ReportData = SessionsToReport.FindByPredicate([&MinimalCrashSession](const FEditorAnalyticReportData& Candidate) { return Candidate.Session.EngineVersion == MinimalCrashSession.EngineVersion; }))
 				{
 					// Account for the crashes that occurred before the analytic was initialized.
-					ReportData->DelayedCrashCount += 1;
+					if (MinimalCrashSession.ExitCode.IsSet())
+					{
+						if (*MinimalCrashSession.ExitCode == 3) // 3 is the engine exit code when a crash is captured and reported.
+						{
+							ReportData->DelayedCrashCount += 1;
+						}
+						else if (PLATFORM_WINDOWS && *ReportData->Session.ExitCode == 1) // On windows, exit code 1 is usually because the user killed the Editor.
+						{
+							ReportData->DelayedKillCount += 1;
+						}
+						else // Any other exit codes.
+						{
+							ReportData->DelayedAbnormalCount += 1;
+						}
+					}
+					else // A minimal crash session was created, so the Editor crashed, but for some reasons, the exit code was unknown.
+					{
+						ReportData->DelayedAbnormalCount += 1;
+					}
 
 					// The minimal crash session was accounted for, schedule it for deletion.
 					SessionsToDelete.Add(MoveTemp(MinimalCrashSession));
@@ -261,8 +285,10 @@ void FEditorSessionSummarySender::SendSessionSummaryEvent(const FEditorAnalyticR
 	AnalyticsAttributes.Emplace(TEXT("IsCrcMissing"), Session.bIsCrcExeMissing);
 	AnalyticsAttributes.Emplace(TEXT("SentFrom"), Sender);
 	AnalyticsAttributes.Emplace(TEXT("MonitorPid"), Session.MonitorProcessID); // For out-of-process monitoring, if this is 0, this will mean CRC failed to launch or crashed very early.
-	AnalyticsAttributes.Emplace(TEXT("DelayedCrashCount"), ReportData.DelayedCrashCount); // Piggybacked information from 'minimal crash sessions', previous Editor instances crashes recorded before analytics was initialized.
-	AnalyticsAttributes.Emplace(TEXT("ProcessDiagnostics"), Session.ProcessDiagnostics);
+	AnalyticsAttributes.Emplace(TEXT("DelayedCrashCount"), ReportData.DelayedCrashCount); // Piggybacked information from 'minimal crash sessions', previous Editor instances that crashed before analytics was initialized.
+	AnalyticsAttributes.Emplace(TEXT("DelayedKillCount"), ReportData.DelayedKillCount);   // Piggybacked information from 'minimal crash sessions', previous Editor instances that were killed before analytics was initialized.
+	AnalyticsAttributes.Emplace(TEXT("DelayedAbnormalCount"), ReportData.DelayedAbnormalCount); // Piggybacked information from 'minimal crash sessions', previous Editor instances that exited abnormally before analytics was initialized.
+	AnalyticsAttributes.Emplace(TEXT("ProcessDiagnostics"), Session.ProcessDiagnostics); // Non-zero means that a diagnostic feature affecting performance was on. (like using a stomp allocator).
 
 	bool bShouldAttachMonitorLog = (ShutdownTypeString != EditorSessionSenderDefs::ShutdownSessionToken && ShutdownTypeString != EditorSessionSenderDefs::TerminatedSessionToken);
 
diff --git a/Engine/Source/Programs/CrashReportClient/Private/CrashReportClientApp.cpp b/Engine/Source/Programs/CrashReportClient/Private/CrashReportClientApp.cpp
--- a/Engine/Source/Programs/CrashReportClient/Private/CrashReportClientApp.cpp
+++ b/Engine/Source/Programs/CrashReportClient/Private/CrashReportClientApp.cpp
@@ -1197,7 +1197,7 @@ void RunCrashReportClient(const TCHAR* CommandLine)
 					else
 					{
 						FDiagnosticLogger::Get().LogEvent(TEXT("MTBF/NoSessionFound"));
-						if (bMonitoredAppCrashed)
+						if (bMonitoredAppCrashed || (MonitoredProcessExitCode.IsSet() && *MonitoredProcessExitCode != 0))
 						{
 							// No session found. Record this crash to piggyback off a valid session later. (Purpose: Account for early Editor crash happening before Analytics was initialized)
 							FEditorAnalyticsSession::CreateMinimalCrashSession(MonitoredProcessExitCode);
