Removed private FEditorSessionSummaryWriter and deprecated public class FEditorSessionSummarySender and FEditorAnalyticsSession.
  - The functionality was generalized in FAnalyticsSessionSummaryManager, FEngineAnalyticsSessionSummary, FEditorAnalyticsSessionSummary and FAnalyticsSessionSummarySender.

Added code to clean up left-over expired files/data using the old analytics system implementation.

diff --git a/Engine/Source/Developer/EditorAnalyticsSession/Private/EditorAnalyticsSession.cpp b/Engine/Source/Developer/EditorAnalyticsSession/Private/EditorAnalyticsSession.cpp
--- a/Engine/Source/Developer/EditorAnalyticsSession/Private/EditorAnalyticsSession.cpp
+++ b/Engine/Source/Developer/EditorAnalyticsSession/Private/EditorAnalyticsSession.cpp
@@ -1,10 +1,16 @@
 // Copyright Epic Games, Inc. All Rights Reserved.
 
 #include "EditorAnalyticsSession.h"
+#include "CoreMinimal.h"
 #include "Modules/ModuleManager.h"
 #include "Internationalization/Regex.h"
+#include "Containers/UnrealString.h"
+#include "GenericPlatform/GenericPlatformFile.h"
+#include "HAL/CriticalSection.h"
 #include "HAL/PlatformProcess.h"
+#include "HAL/PlatformMisc.h"
 #include "HAL/FileManager.h"
+#include "Misc/DateTime.h"
 #include "Misc/EngineVersion.h"
 #include "Misc/Guid.h"
 #include "Misc/Paths.h"
@@ -13,16 +19,6 @@ IMPLEMENT_MODULE(FEditorAnalyticsSessionModule, EditorAnalyticsSession);
 
 namespace EditorAnalyticsDefs
 {
-	static const FString FalseValueString(TEXT("0"));
-	static const FString TrueValueString(TEXT("1"));
-
-	static const FString DefaultUserActivity(TEXT("Unknown"));
-	static const FString UnknownProjectValueString(TEXT("UnknownProject"));
-
-	static const FString UnknownAppIdString(TEXT("UnknownAppId"));
-	static const FString UnknownAppVersionString(TEXT("UnknownAppVersion"));
-	static const FString UnknownUserIdString(TEXT("UnknownUserID"));
-
 	// The storage location is used to version the different data format. This is to prevent one version of the Editor/CRC to send sessions produced by another incompatible version.
 	//   Version 1_0 : Used from creation up to 4.25.0 release (included).
 	//   Version 1_1 : To avoid public API changes in 4.25.1, TotalUserInactivitySeconds was repurposed to contain the SessionDuration read from FPlatformTime::Seconds() to detect cases where the user system date time is unreliable.
@@ -45,100 +41,12 @@ namespace EditorAnalyticsDefs
 	static const FString SessionSummarySection = SessionSummaryRoot / TEXT("1_7"); // The current session format.
 	static const FString GlobalLockName(TEXT("UE4_SessionSummary_Lock"));
 	static const FString SessionListStoreKey(TEXT("SessionList"));
-
-	// capture context
-	static const FString AppIdStoreKey(TEXT("AppId"));
-	static const FString AppVersionStoreKey(TEXT("AppVersion"));
-	static const FString UserIdStoreKey(TEXT("UserId"));
-
-	// general values
-	static const FString ProjectNameStoreKey(TEXT("ProjectName"));
-	static const FString ProjectIDStoreKey(TEXT("ProjectID"));
-	static const FString ProjectDescriptionStoreKey(TEXT("ProjectDescription"));
-	static const FString ProjectVersionStoreKey(TEXT("ProjectVersion"));
-	static const FString EngineVersionStoreKey(TEXT("EngineVersion"));
-	static const FString PlatformProcessIDStoreKey(TEXT("PlatformProcessID"));
-	static const FString MonitorProcessIDStoreKey(TEXT("MonitorProcessID"));
-	static const FString ExitCodeStoreKey(TEXT("ExitCode"));
-	static const FString MonitorExceptCodeStoreKey(TEXT("MonitorExceptCode"));
-	static const FString MonitorExitCodeStoreKey(TEXT("MonitorExitCode"));
-	static const FString SessionTickCountStoreKey(TEXT("SessionTickCount"));
-	static const FString EngineTickCountStoreKey(TEXT("EngineTickCount"));
-	static const FString UserInteractionCountStoreKey(TEXT("UserInteractionCount"));
-	static const FString CommandLineStoreKey(TEXT("CommandLine"));
-
-	// stall detector
-	static const FString TotalStallCountStoreKey(TEXT("TotalStallCount"));
-	static const FString TotalStallReportedStoreKey(TEXT("TotalStallReported"));
-	static const FString TopStallNameStoreKey(TEXT("TopStallName"));
-	static const FString TopStallBudgetSecondsStoreKey(TEXT("TopStallBudgetSeconds"));
-	static const FString TopStallOverageSecondsStoreKey(TEXT("TopStallOverageSeconds"));
-	static const FString TopStallTriggerCountStoreKey(TEXT("TopStallTriggerCount"));
-
-	// timestamps
-	static const FString StartupTimestampStoreKey(TEXT("StartupTimestamp"));
 	static const FString TimestampStoreKey(TEXT("Timestamp"));
-	static const FString SessionDurationStoreKey(TEXT("SessionDuration"));
-	static const FString Idle1MinStoreKey(TEXT("Idle1Min"));
-	static const FString Idle5MinStoreKey(TEXT("Idle5Min"));
-	static const FString Idle30MinStoreKey(TEXT("Idle30Min"));
-	static const FString TotalEditorInactivitySecondsStoreKey(TEXT("TotalEditorInactivitySecs"));
-	static const FString CurrentUserActivityStoreKey(TEXT("CurrentUserActivity"));
-	static const FString PluginsStoreKey(TEXT("Plugins"));
-	static const FString AverageFPSStoreKey(TEXT("AverageFPS"));
-	static const FString LastTickTimestampStoreKey(TEXT("LastTickTimestamp"));
-	static const FString DeathTimestampStoreKey(TEXT("DeathTimestamp"));
-
-	// GPU details
-	static const FString DesktopGPUAdapterStoreKey(TEXT("DesktopGPUAdapter"));
-	static const FString RenderingGPUAdapterStoreKey(TEXT("RenderingGPUAdapter"));
-	static const FString GPUVendorIDStoreKey(TEXT("GPUVendorID"));
-	static const FString GPUDeviceIDStoreKey(TEXT("GPUDeviceID"));
-	static const FString GRHIDeviceRevisionStoreKey(TEXT("GRHIDeviceRevision"));
-	static const FString GRHIAdapterInternalDriverVersionStoreKey(TEXT("GRHIAdapterInternalDriverVersion"));
-	static const FString GRHIAdapterUserDriverVersionStoreKey(TEXT("GRHIAdapterUserDriverVersion"));
-	static const FString GRHINameStoreKey(TEXT("GRHIName"));
-
-	// CPU details
-	static const FString TotalPhysicalRAMStoreKey(TEXT("TotalPhysicalRAM"));
-	static const FString CPUPhysicalCoresStoreKey(TEXT("CPUPhysicalCores"));
-	static const FString CPULogicalCoresStoreKey(TEXT("CPULogicalCores"));
-	static const FString CPUVendorStoreKey(TEXT("CPUVendor"));
-	static const FString CPUBrandStoreKey(TEXT("CPUBrand"));
-
-	// OS details
-	static const FString OSMajorStoreKey(TEXT("OSMajor"));
-	static const FString OSMinorStoreKey(TEXT("OSMinor"));
-	static const FString OSVersionStoreKey(TEXT("OSVersion"));
-	static const FString bIs64BitOSStoreKey(TEXT("bIs64BitOS"));
-
-	// boolean flags
-	static const FString IsCrashStoreKey(TEXT("IsCrash"));
-	static const FString IsGPUCrashStoreKey(TEXT("IsGPUCrash"));
-	static const FString IsDebuggerPresentStoreKey(TEXT("IsDebuggerPresent"));
-	static const FString IsDebuggerIgnoredStoreKey(TEXT("IsDebuggerIgnored"));
-	static const FString WasDebuggerStoreKey(TEXT("WasEverDebugger"));
-	static const FString IsVanillaStoreKey(TEXT("IsVanilla"));
-	static const FString IsTerminatingStoreKey(TEXT("Terminating"));
-	static const FString WasShutdownStoreKey(TEXT("WasShutdown"));
-	static const FString IsUserLoggingOutStoreKey(TEXT("IsUserLoggingOut"));
-	static const FString IsInPIEStoreKey(TEXT("IsInPIE"));
-	static const FString IsInEnterpriseStoreKey(TEXT("IsInEnterprise"));
-	static const FString IsInVRModeStoreKey(TEXT("IsInVRMode"));
-	static const FString IsCrcExeMissingStoreKey(TEXT("IsCrcExeMissing"));
-	static const FString IsLowDriveSpaceStoreKey(TEXT("IsLowDriveSpace"));
-
-	static const FString ProcessDiagnosticsStoreKey(TEXT("ProcessDiagnostics"));
 }
 
 // Utilities for writing to stored values
 namespace EditorAnalyticsUtils
 {
-	static FString TimestampToString(FDateTime InTimestamp)
-	{
-		return LexToString(InTimestamp.ToUnixTimestamp());
-	}
-
 	static FDateTime StringToTimestamp(FString InString)
 	{
 		int64 TimestampUnix;
@@ -149,167 +57,20 @@ namespace EditorAnalyticsUtils
 		return FDateTime::MinValue();
 	}
 
-	static FString BoolToStoredString(bool bInValue)
-	{
-		return bInValue ? EditorAnalyticsDefs::TrueValueString : EditorAnalyticsDefs::FalseValueString;
-	}
-
-	static bool GetStoredBool(const FString& SectionName, const FString& StoredKey)
-	{
-		FString StoredString = EditorAnalyticsDefs::FalseValueString;
-		FPlatformMisc::GetStoredValue(EditorAnalyticsDefs::StoreId, SectionName, StoredKey, StoredString);
-
-		return StoredString == EditorAnalyticsDefs::TrueValueString;
-	}
-
-	static FString GetSessionStorageLocation(const FString& SessionID)
-	{
-		return EditorAnalyticsDefs::SessionSummarySection + TEXT("/") + SessionID;
-	}
-
 	static FString GetSessionEventLogDir()
 	{
 		return FString::Printf(TEXT("%sAnalytics"), FPlatformProcess::ApplicationSettingsDir());
 	}
 
-	static bool IsSavingIndividualFieldFastAndThreadSafe()
-	{
-	#if PLATFORM_WINDOWS
-		return true; // Saving individual field to registry on Windows is fast and thread safe.
-	#else
-		return false; // On other platforms, we load/update/save an entire .ini file. This is not fast and not thread safe.
-	#endif
-	}
-
-	static void LogSessionEvent(FEditorAnalyticsSession& Session, FEditorAnalyticsSession::EEventType InEventType, const FDateTime& InTimestamp)
-	{
-		// This is primary logging mechanims. It works across all platforms and is rather robust. It uses the robustness of the file system to log events. It creates a file
-		// (a directory) for each event and encodes the event payload in the file name. This doesn't require any fancy synchronization or complicated concurrent file IO
-		// implementation. Since the number of events is low (0 to 5 per session), that's a straight forward working solution. The files are deleted when the session is
-		// deleted. Also avoid memory allocation, this can be called when the heap is corrupted.
-		TCHAR TimestampStr[256];
-		FCString::Sprintf(TimestampStr, TFormatSpecifier<decltype(InTimestamp.ToUnixTimestamp())>::GetFormatSpecifier(), InTimestamp.ToUnixTimestamp());
-
-		TCHAR Pathname[512];
-		FCString::Sprintf(Pathname, TEXT("%s/%s_%d_%d_%d_%d_%d_%s"),
-			*EditorAnalyticsUtils::GetSessionEventLogDir(),
-			*Session.SessionId,
-			static_cast<int32>(InEventType),
-			FPlatformAtomics::AtomicRead(&Session.Idle1Min),
-			FPlatformAtomics::AtomicRead(&Session.Idle5Min),
-			FPlatformAtomics::AtomicRead(&Session.Idle30Min),
-			FPlatformAtomics::AtomicRead(&Session.SessionDuration),
-			TimestampStr);
-
-		IFileManager::Get().MakeDirectory(Pathname, /*Tree*/true);
-
-		// As a secondary/backup mechanism, directly save the info in the session if this is fast and safe. This allocate memory, so that's not perfect in case of crash.
-		if (IsSavingIndividualFieldFastAndThreadSafe())
-		{
-			const FString StorageLocation = EditorAnalyticsUtils::GetSessionStorageLocation(Session.SessionId);
-
-			switch(InEventType)
-			{
-			case FEditorAnalyticsSession::EEventType::Crashed:
-				FPlatformMisc::SetStoredValue(EditorAnalyticsDefs::StoreId, StorageLocation, EditorAnalyticsDefs::IsCrashStoreKey, EditorAnalyticsUtils::BoolToStoredString(true));
-				break;
-
-			case FEditorAnalyticsSession::EEventType::GpuCrashed:
-				FPlatformMisc::SetStoredValue(EditorAnalyticsDefs::StoreId, StorageLocation, EditorAnalyticsDefs::IsGPUCrashStoreKey, EditorAnalyticsUtils::BoolToStoredString(true));
-				break;
-
-			case FEditorAnalyticsSession::EEventType::Terminated:
-				FPlatformMisc::SetStoredValue(EditorAnalyticsDefs::StoreId, StorageLocation, EditorAnalyticsDefs::IsTerminatingStoreKey, EditorAnalyticsUtils::BoolToStoredString(true));
-				break;
-
-			case FEditorAnalyticsSession::EEventType::Shutdown:
-				FPlatformMisc::SetStoredValue(EditorAnalyticsDefs::StoreId, StorageLocation, EditorAnalyticsDefs::WasShutdownStoreKey, EditorAnalyticsUtils::BoolToStoredString(true));
-				break;
-
-			case FEditorAnalyticsSession::EEventType::LogOut:
-				FPlatformMisc::SetStoredValue(EditorAnalyticsDefs::StoreId, StorageLocation, EditorAnalyticsDefs::IsUserLoggingOutStoreKey, EditorAnalyticsUtils::BoolToStoredString(true));
-				break;
-
-			default:
-				break;
-			}
-
-			// Save the timestamps.
-			FPlatformMisc::SetStoredValue(EditorAnalyticsDefs::StoreId, StorageLocation, EditorAnalyticsDefs::Idle1MinStoreKey, FString::FromInt(FPlatformAtomics::AtomicRead(&Session.Idle1Min)));
-			FPlatformMisc::SetStoredValue(EditorAnalyticsDefs::StoreId, StorageLocation, EditorAnalyticsDefs::Idle5MinStoreKey, FString::FromInt(FPlatformAtomics::AtomicRead(&Session.Idle5Min)));
-			FPlatformMisc::SetStoredValue(EditorAnalyticsDefs::StoreId, StorageLocation, EditorAnalyticsDefs::Idle30MinStoreKey, FString::FromInt(FPlatformAtomics::AtomicRead(&Session.Idle30Min)));
-			FPlatformMisc::SetStoredValue(EditorAnalyticsDefs::StoreId, StorageLocation, EditorAnalyticsDefs::SessionDurationStoreKey, FString::FromInt(FPlatformAtomics::AtomicRead(&Session.SessionDuration)));
-			FPlatformMisc::SetStoredValue(EditorAnalyticsDefs::StoreId, StorageLocation, EditorAnalyticsDefs::TimestampStoreKey, EditorAnalyticsUtils::TimestampToString(InTimestamp));
-		}
-	}
-
-	/** Analyze the events logged with LogEvent() and update the session fields to reflect the last state of the session. */
-	static void UpdateSessionFromLogAnalysis(FEditorAnalyticsSession& Session)
-	{
-		// Read and aggregate the log events. The event data is encoded in the directory names created by the logger
-		FRegexPattern Pattern(TEXT(R"((^[a-fA-F0-9-]+)_([0-9]+)_([0-9]+)_([0-9]+)_([0-9]+)_([0-9]+)_([0-9]+))")); // Need help with regex? Try https://regex101.com/
-		IFileManager::Get().IterateDirectoryRecursively(*EditorAnalyticsUtils::GetSessionEventLogDir(), [&Session, &Pattern](const TCHAR* Pathname, bool bIsDir)
-		{
-			if (bIsDir)
-			{
-				FRegexMatcher Matcher(Pattern, FPaths::GetCleanFilename(Pathname));
-				if (Matcher.FindNext() && Matcher.GetCaptureGroup(1) == Session.SessionId)
-				{
-					FEditorAnalyticsSession::EEventType EventType = static_cast<FEditorAnalyticsSession::EEventType>(FCString::Atoi(*Matcher.GetCaptureGroup(2))); // Event
-					switch (EventType)
-					{
-						case FEditorAnalyticsSession::EEventType::Crashed:     Session.bCrashed = true;          break;
-						case FEditorAnalyticsSession::EEventType::GpuCrashed : Session.bGPUCrashed = true;       break;
-						case FEditorAnalyticsSession::EEventType::Terminated:  Session.bIsTerminating = true;    break;
-						case FEditorAnalyticsSession::EEventType::Shutdown:    Session.bWasShutdown = true;      break;
-						case FEditorAnalyticsSession::EEventType::LogOut:      Session.bIsUserLoggingOut = true; break;
-						default: break;
-					}
-
-					int32 ParsedTime = FCString::Atoi(*Matcher.GetCaptureGroup(3)); // Idle1Min.
-					if (ParsedTime > Session.Idle1Min)
-					{
-						Session.Idle1Min = ParsedTime; // No concurrency expected when reloading (no need for atomic compare exchange)
-					}
-
-					ParsedTime = FCString::Atoi(*Matcher.GetCaptureGroup(4)); // Idle5Min.
-					if (ParsedTime > Session.Idle5Min)
-					{
-						Session.Idle5Min = ParsedTime;
-					}
-
-					ParsedTime = FCString::Atoi(*Matcher.GetCaptureGroup(5)); // Idle30Min.
-					if (ParsedTime > Session.Idle30Min)
-					{
-						Session.Idle30Min = ParsedTime;
-					}
-
-					ParsedTime = FCString::Atoi(*Matcher.GetCaptureGroup(6)); // SessionDuration.
-					if (ParsedTime > Session.SessionDuration)
-					{
-						Session.SessionDuration = ParsedTime;
-					}
-
-					FDateTime ParsedTimestamp = FDateTime::FromUnixTimestamp(FCString::Atoi64(*Matcher.GetCaptureGroup(6))); // Unix timestamp (UTC)
-					if (ParsedTimestamp > Session.Timestamp)
-					{
-						Session.Timestamp = ParsedTimestamp;
-					}
-				}
-			}
-			return true;
-		});
-	}
-
-	static void DeleteLogEvents(const FEditorAnalyticsSession& Session)
+	static void DeleteLogEvents(const FString& SessionId)
 	{
 		// Gather the list of files
 		TArray<FString> SessionEventPaths;
-		IFileManager::Get().IterateDirectoryRecursively(*EditorAnalyticsUtils::GetSessionEventLogDir(), [&Session, &SessionEventPaths](const TCHAR* Pathname, bool bIsDir)
+		IFileManager::Get().IterateDirectoryRecursively(*EditorAnalyticsUtils::GetSessionEventLogDir(), [&SessionId, &SessionEventPaths](const TCHAR* Pathname, bool bIsDir)
 		{
 			if (bIsDir)
 			{
-				if (FPaths::GetCleanFilename(Pathname).StartsWith(Session.SessionId))
+				if (FPaths::GetCleanFilename(Pathname).StartsWith(SessionId))
 				{
 					SessionEventPaths.Emplace(Pathname);
 				}
@@ -324,175 +85,6 @@ namespace EditorAnalyticsUtils
 		}
 	}
 
-// Utility macros to make it easier to check that all fields are being written to.
-#define GET_STORED_STRING(FieldName) \
-	{ \
-		FPlatformMisc::GetStoredValue(EditorAnalyticsDefs::StoreId, SectionName, EditorAnalyticsDefs:: FieldName ## StoreKey, Session.FieldName); \
-	}
-
-#define GET_STORED_INT(FieldName) \
-	{ \
-		FString FieldName ## Temp; \
-		FPlatformMisc::GetStoredValue(EditorAnalyticsDefs::StoreId, SectionName, EditorAnalyticsDefs:: FieldName ## StoreKey, FieldName ## Temp); \
-		Session.FieldName = FCString::Atoi(*FieldName ## Temp); \
-	}
-
-#define GET_STORED_FLOAT(FieldName) \
-	{ \
-		FString FieldName ## Temp; \
-		FPlatformMisc::GetStoredValue(EditorAnalyticsDefs::StoreId, SectionName, EditorAnalyticsDefs:: FieldName ## StoreKey, FieldName ## Temp); \
-		Session.FieldName = FCString::Atof(*FieldName ## Temp); \
-	}
-
-	static void LoadInternal(FEditorAnalyticsSession& Session, const FString& InSessionId)
-	{
-		Session.SessionId = InSessionId;
-
-		FString SectionName = EditorAnalyticsUtils::GetSessionStorageLocation(Session.SessionId);
-
-		GET_STORED_STRING(AppId);
-		GET_STORED_STRING(AppVersion);
-		GET_STORED_STRING(UserId);
-
-		GET_STORED_STRING(ProjectName);
-		GET_STORED_STRING(ProjectID);
-		GET_STORED_STRING(ProjectDescription);
-		GET_STORED_STRING(ProjectVersion);
-		GET_STORED_STRING(EngineVersion);
-		GET_STORED_STRING(CommandLine);
-		GET_STORED_INT(PlatformProcessID);
-		GET_STORED_INT(MonitorProcessID);
-
-		{
-			FString ExitCodeString;
-			if (FPlatformMisc::GetStoredValue(EditorAnalyticsDefs::StoreId, SectionName, EditorAnalyticsDefs::ExitCodeStoreKey, ExitCodeString))
-			{
-				Session.ExitCode.Emplace(FCString::Atoi(*ExitCodeString));
-			}
-		}
-
-		{
-			FString MonitorExceptCodeString;
-			if (FPlatformMisc::GetStoredValue(EditorAnalyticsDefs::StoreId, SectionName, EditorAnalyticsDefs::MonitorExceptCodeStoreKey, MonitorExceptCodeString))
-			{
-				Session.MonitorExceptCode.Emplace(FCString::Atoi(*MonitorExceptCodeString));
-			}
-		}
-
-		{
-			FString MonitorExitCodeString;
-			if (FPlatformMisc::GetStoredValue(EditorAnalyticsDefs::StoreId, SectionName, EditorAnalyticsDefs::MonitorExitCodeStoreKey, MonitorExitCodeString))
-			{
-				Session.MonitorExitCode.Emplace(FCString::Atoi(*MonitorExitCodeString));
-			}
-		}
-
-		// scope is just to isolate the temporary value
-		{
-			FString StartupTimestampString;
-			FPlatformMisc::GetStoredValue(EditorAnalyticsDefs::StoreId, SectionName, EditorAnalyticsDefs::StartupTimestampStoreKey, StartupTimestampString);
-			Session.StartupTimestamp = EditorAnalyticsUtils::StringToTimestamp(StartupTimestampString);
-		}
-
-		{
-			FString TimestampString;
-			FPlatformMisc::GetStoredValue(EditorAnalyticsDefs::StoreId, SectionName, EditorAnalyticsDefs::TimestampStoreKey, TimestampString);
-			Session.Timestamp = EditorAnalyticsUtils::StringToTimestamp(TimestampString);
-		}
-
-		{
-			FString LastTickTimestampString;
-			FPlatformMisc::GetStoredValue(EditorAnalyticsDefs::StoreId, SectionName, EditorAnalyticsDefs::LastTickTimestampStoreKey, LastTickTimestampString);
-			Session.LastTickTimestamp = EditorAnalyticsUtils::StringToTimestamp(LastTickTimestampString);
-		}
-
-		{
-			FString DeathTimestampString;
-			if (FPlatformMisc::GetStoredValue(EditorAnalyticsDefs::StoreId, SectionName, EditorAnalyticsDefs::DeathTimestampStoreKey, DeathTimestampString))
-			{
-				Session.DeathTimestamp.Emplace(EditorAnalyticsUtils::StringToTimestamp(DeathTimestampString));
-			}
-		}
-
-		GET_STORED_INT(SessionDuration);
-		GET_STORED_INT(Idle1Min);
-		GET_STORED_INT(Idle5Min);
-		GET_STORED_INT(Idle30Min);
-		GET_STORED_INT(TotalEditorInactivitySeconds);
-
-		GET_STORED_STRING(CurrentUserActivity);
-
-		{
-			FString PluginsString;
-			FPlatformMisc::GetStoredValue(EditorAnalyticsDefs::StoreId, SectionName, EditorAnalyticsDefs::PluginsStoreKey, PluginsString);
-			PluginsString.ParseIntoArray(Session.Plugins, TEXT(","));
-		}
-
-		GET_STORED_FLOAT(AverageFPS);
-
-		GET_STORED_INT(SessionTickCount);
-		GET_STORED_INT(EngineTickCount);
-		GET_STORED_INT(UserInteractionCount);
-
-		GET_STORED_INT(TotalStallCount);
-		GET_STORED_INT(TotalStallReported);
-		GET_STORED_STRING(TopStallName);
-		GET_STORED_FLOAT(TopStallBudgetSeconds);
-		GET_STORED_FLOAT(TopStallOverageSeconds);
-		GET_STORED_INT(TopStallTriggerCount);
-
-		GET_STORED_STRING(DesktopGPUAdapter);
-		GET_STORED_STRING(RenderingGPUAdapter);
-
-		GET_STORED_INT(GPUVendorID);
-		GET_STORED_INT(GPUDeviceID);
-		GET_STORED_INT(GRHIDeviceRevision);
-
-		GET_STORED_STRING(GRHIAdapterInternalDriverVersion);
-		GET_STORED_STRING(GRHIAdapterUserDriverVersion);
-		GET_STORED_STRING(GRHIName);
-
-		{
-			FString TotalPhysicalRAMString;
-			FPlatformMisc::GetStoredValue(EditorAnalyticsDefs::StoreId, SectionName, EditorAnalyticsDefs::TotalPhysicalRAMStoreKey, TotalPhysicalRAMString);
-			Session.TotalPhysicalRAM = FCString::Atoi64(*TotalPhysicalRAMString);
-		}
-
-		GET_STORED_INT(CPUPhysicalCores);
-		GET_STORED_INT(CPULogicalCores);
-
-		GET_STORED_STRING(CPUVendor);
-		GET_STORED_STRING(CPUBrand);
-
-		GET_STORED_STRING(OSMajor);
-		GET_STORED_STRING(OSMinor);
-		GET_STORED_STRING(OSVersion);
-
-		Session.bIs64BitOS = EditorAnalyticsUtils::GetStoredBool(SectionName, EditorAnalyticsDefs::bIs64BitOSStoreKey);
-		Session.bCrashed = EditorAnalyticsUtils::GetStoredBool(SectionName, EditorAnalyticsDefs::IsCrashStoreKey);
-		Session.bGPUCrashed = EditorAnalyticsUtils::GetStoredBool(SectionName, EditorAnalyticsDefs::IsGPUCrashStoreKey);
-		Session.bIsDebuggerPresent = EditorAnalyticsUtils::GetStoredBool(SectionName, EditorAnalyticsDefs::IsDebuggerPresentStoreKey);
-		Session.bIsDebuggerIgnored = EditorAnalyticsUtils::GetStoredBool(SectionName, EditorAnalyticsDefs::IsDebuggerIgnoredStoreKey);
-		Session.bWasEverDebugger = EditorAnalyticsUtils::GetStoredBool(SectionName, EditorAnalyticsDefs::WasDebuggerStoreKey);
-		Session.bIsVanilla = EditorAnalyticsUtils::GetStoredBool(SectionName, EditorAnalyticsDefs::IsVanillaStoreKey);
-		Session.bIsTerminating = EditorAnalyticsUtils::GetStoredBool(SectionName, EditorAnalyticsDefs::IsTerminatingStoreKey);
-		Session.bWasShutdown = EditorAnalyticsUtils::GetStoredBool(SectionName, EditorAnalyticsDefs::WasShutdownStoreKey);
-		Session.bIsUserLoggingOut = EditorAnalyticsUtils::GetStoredBool(SectionName, EditorAnalyticsDefs::IsUserLoggingOutStoreKey);
-		Session.bIsInPIE = EditorAnalyticsUtils::GetStoredBool(SectionName, EditorAnalyticsDefs::IsInPIEStoreKey);
-		Session.bIsInVRMode = EditorAnalyticsUtils::GetStoredBool(SectionName, EditorAnalyticsDefs::IsInVRModeStoreKey);
-		Session.bIsInEnterprise = EditorAnalyticsUtils::GetStoredBool(SectionName, EditorAnalyticsDefs::IsInEnterpriseStoreKey);
-		Session.bIsLowDriveSpace = EditorAnalyticsUtils::GetStoredBool(SectionName, EditorAnalyticsDefs::IsLowDriveSpaceStoreKey);
-		Session.bIsCrcExeMissing = EditorAnalyticsUtils::GetStoredBool(SectionName, EditorAnalyticsDefs::IsCrcExeMissingStoreKey);
-
-		GET_STORED_INT(ProcessDiagnostics);
-
-		// Analyze the logged events and update corresponding fields in the session.
-		UpdateSessionFromLogAnalysis(Session);
-	}
-
-#undef GET_STORED_INT
-#undef GET_STORED_STRING
-
 	static TArray<FString> GetSessionList()
 	{
 		FString SessionListString;
@@ -505,357 +97,12 @@ namespace EditorAnalyticsUtils
 	}
 }
 
-TUniquePtr<FSystemWideCriticalSection> FEditorAnalyticsSession::StoredValuesLock;
-std::atomic<uint64> FEditorAnalyticsSession::StoredValuesLockOwnerInfo(0);
-
-FEditorAnalyticsSession::FEditorAnalyticsSession()
-{
-	AppId = EditorAnalyticsDefs::UnknownAppIdString;
-	AppVersion = EditorAnalyticsDefs::UnknownAppVersionString;
-	UserId = EditorAnalyticsDefs::UnknownUserIdString;
-
-	ProjectName = EditorAnalyticsDefs::UnknownProjectValueString;
-	StartupTimestamp = FDateTime::MinValue();
-	Timestamp = FDateTime::MinValue();
-	CurrentUserActivity = EditorAnalyticsDefs::DefaultUserActivity;
-
-	bIs64BitOS = false;
-	bCrashed = false;
-	bGPUCrashed = false;
-	bIsDebuggerPresent = false;
-	bWasEverDebugger = false;
-	bIsVanilla = false;
-	bIsTerminating = false;
-	bWasShutdown = false;
-	bIsUserLoggingOut = false;
-	bIsInPIE = false;
-	bIsInEnterprise = false;
-	bIsInVRMode = false;
-	bAlreadySaved = false;
-	bIsLowDriveSpace = false;
-	bIsCrcExeMissing = false;
-	bIsDebuggerIgnored = false;
-
-	ProcessDiagnostics = 0x0;
-}
-
-bool FEditorAnalyticsSession::Lock(FTimespan Timeout)
-{
-	// If Lock() is called recursively from a thread already owning the lock, fire an ensure. FSystemWideCriticalSection doesn't support recursive locking correctly.
-	if (!ensure(!IsLockedBy(FPlatformTLS::GetCurrentThreadId())))
-	{
-		return true;
-	}
-
-	// Try to acquire the lock.
-	TUniquePtr<FSystemWideCriticalSection> TempLock = MakeUnique<FSystemWideCriticalSection>(EditorAnalyticsDefs::GlobalLockName, Timeout);
-
-	// If the lock is acquired.
-	if (TempLock && TempLock->IsValid())
-	{
-		// The lock was successfully acquired by this thread, keep it until Unlock() is called.
-		StoredValuesLock = MoveTemp(TempLock);
-		static_assert(sizeof(FPlatformTLS::GetCurrentThreadId()) == sizeof(uint32), "Encoding below assume the thread Id only uses 32 bits");
-		StoredValuesLockOwnerInfo.store(((1ull << (sizeof(uint32) * 8))) | FPlatformTLS::GetCurrentThreadId()); // 32 bits for thread id, 1 bit for 'is taken' so that threadId 0 is supported.
-		return true;
-	}
-
-	return false; // Failed to get the lock within the allowed time.
-}
-
-void FEditorAnalyticsSession::Unlock()
-{
-	// If Unlock() is called from a non-owning thread or recursively from an owning thread, fire an ensure. FSystemWideCriticalSection doesn't support recursive locking correctly.
-	if (!ensure(IsLockedBy(FPlatformTLS::GetCurrentThreadId())))
-	{
-		return;
-	}
-
-	StoredValuesLockOwnerInfo.store(0);
-	StoredValuesLock.Reset();
-}
-
-bool FEditorAnalyticsSession::IsLockedBy(uint32 ThreadId)
-{
-	// Check if the calling thread has the lock.
-	uint64 LockInfo = StoredValuesLockOwnerInfo.load();
-	return LockInfo != 0 && (LockInfo & 0xFFFFFFFF) == ThreadId;
-}
-
-bool FEditorAnalyticsSession::Save()
-{
-	if (!ensure(IsLockedBy(FPlatformTLS::GetCurrentThreadId())))
-	{
-		return false;
-	}
-
-	const FString StorageLocation = EditorAnalyticsUtils::GetSessionStorageLocation(SessionId);
-
-	if (!bAlreadySaved)
-	{
-		const FString PluginsString = FString::Join(Plugins, TEXT(","));
-
-		TMap<FString, FString> KeyValues = {
-			{EditorAnalyticsDefs::EngineVersionStoreKey,       EngineVersion},
-			{EditorAnalyticsDefs::CommandLineStoreKey,         CommandLine},
-			{EditorAnalyticsDefs::PlatformProcessIDStoreKey,   FString::FromInt(PlatformProcessID)},
-			{EditorAnalyticsDefs::MonitorProcessIDStoreKey,    FString::FromInt(MonitorProcessID)},
-			{EditorAnalyticsDefs::DesktopGPUAdapterStoreKey,   DesktopGPUAdapter},
-			{EditorAnalyticsDefs::RenderingGPUAdapterStoreKey, RenderingGPUAdapter},
-			{EditorAnalyticsDefs::GPUVendorIDStoreKey,         FString::FromInt(GPUVendorID)},
-			{EditorAnalyticsDefs::GPUDeviceIDStoreKey,         FString::FromInt(GPUDeviceID)},
-			{EditorAnalyticsDefs::GRHIDeviceRevisionStoreKey,  FString::FromInt(GRHIDeviceRevision)},
-			{EditorAnalyticsDefs::GRHIAdapterInternalDriverVersionStoreKey, GRHIAdapterUserDriverVersion},
-			{EditorAnalyticsDefs::GRHIAdapterUserDriverVersionStoreKey,     GRHIAdapterUserDriverVersion},
-			{EditorAnalyticsDefs::GRHINameStoreKey,                         GRHIName},
-			{EditorAnalyticsDefs::TotalPhysicalRAMStoreKey, FString::Printf(TEXT("%llu"), TotalPhysicalRAM)},
-			{EditorAnalyticsDefs::CPUPhysicalCoresStoreKey, FString::FromInt(CPUPhysicalCores)},
-			{EditorAnalyticsDefs::CPULogicalCoresStoreKey,  FString::FromInt(CPULogicalCores)},
-			{EditorAnalyticsDefs::CPUVendorStoreKey,        CPUVendor},
-			{EditorAnalyticsDefs::CPUBrandStoreKey,         CPUBrand},
-			{EditorAnalyticsDefs::StartupTimestampStoreKey, EditorAnalyticsUtils::TimestampToString(StartupTimestamp)},
-			{EditorAnalyticsDefs::OSMajorStoreKey,    OSMajor},
-			{EditorAnalyticsDefs::OSMinorStoreKey,    OSMinor},
-			{EditorAnalyticsDefs::OSVersionStoreKey,  OSVersion},
-			{EditorAnalyticsDefs::bIs64BitOSStoreKey, EditorAnalyticsUtils::BoolToStoredString(bIs64BitOS)},
-			{EditorAnalyticsDefs::IsCrcExeMissingStoreKey, EditorAnalyticsUtils::BoolToStoredString(bIsCrcExeMissing)},
-			{EditorAnalyticsDefs::PluginsStoreKey,    PluginsString},
-			{EditorAnalyticsDefs::AppIdStoreKey,      AppId},
-			{EditorAnalyticsDefs::AppVersionStoreKey, AppVersion},
-			{EditorAnalyticsDefs::UserIdStoreKey,     UserId},
-		};
-
-		FPlatformMisc::SetStoredValues(EditorAnalyticsDefs::StoreId, StorageLocation, KeyValues);
-
-		bAlreadySaved = true;
-	}
-
-	{
-		TMap<FString, FString> KeyValues = {
-			{EditorAnalyticsDefs::ProjectNameStoreKey, ProjectName},
-			{EditorAnalyticsDefs::ProjectIDStoreKey,   ProjectID},
-			{EditorAnalyticsDefs::ProjectDescriptionStoreKey,  ProjectDescription},
-			{EditorAnalyticsDefs::ProjectVersionStoreKey,  ProjectVersion},
-			{EditorAnalyticsDefs::TimestampStoreKey,         EditorAnalyticsUtils::TimestampToString(Timestamp)},
-			{EditorAnalyticsDefs::LastTickTimestampStoreKey, EditorAnalyticsUtils::TimestampToString(LastTickTimestamp)},
-			{EditorAnalyticsDefs::SessionDurationStoreKey,  FString::FromInt(SessionDuration)},
-			{EditorAnalyticsDefs::Idle1MinStoreKey,  FString::FromInt(Idle1Min)},
-			{EditorAnalyticsDefs::Idle5MinStoreKey,  FString::FromInt(Idle5Min)},
-			{EditorAnalyticsDefs::Idle30MinStoreKey, FString::FromInt(Idle30Min)},
-			{EditorAnalyticsDefs::TotalEditorInactivitySecondsStoreKey, FString::FromInt(TotalEditorInactivitySeconds)},
-			{EditorAnalyticsDefs::CurrentUserActivityStoreKey, CurrentUserActivity},
-			{EditorAnalyticsDefs::AverageFPSStoreKey,       FString::SanitizeFloat(AverageFPS)},
-			{EditorAnalyticsDefs::SessionTickCountStoreKey, FString::FromInt(SessionTickCount)},
-			{EditorAnalyticsDefs::EngineTickCountStoreKey,  FString::FromInt(EngineTickCount)},
-			{EditorAnalyticsDefs::UserInteractionCountStoreKey, FString::FromInt(UserInteractionCount)},
-			{EditorAnalyticsDefs::TotalStallCountStoreKey,        FString::FromInt(TotalStallCount)},
-			{EditorAnalyticsDefs::TotalStallReportedStoreKey,     FString::FromInt(TotalStallReported)},
-			{EditorAnalyticsDefs::TopStallNameStoreKey,           TopStallName},
-			{EditorAnalyticsDefs::TopStallBudgetSecondsStoreKey,  FString::SanitizeFloat(TopStallBudgetSeconds)},
-			{EditorAnalyticsDefs::TopStallOverageSecondsStoreKey, FString::SanitizeFloat(TopStallOverageSeconds)},
-			{EditorAnalyticsDefs::TopStallTriggerCountStoreKey,   FString::FromInt(TopStallTriggerCount)},
-			{EditorAnalyticsDefs::IsDebuggerPresentStoreKey,EditorAnalyticsUtils::BoolToStoredString(bIsDebuggerPresent)},
-			{EditorAnalyticsDefs::IsDebuggerIgnoredStoreKey,EditorAnalyticsUtils::BoolToStoredString(bIsDebuggerIgnored)},
-			{EditorAnalyticsDefs::WasDebuggerStoreKey,      EditorAnalyticsUtils::BoolToStoredString(bWasEverDebugger)},
-			{EditorAnalyticsDefs::IsVanillaStoreKey,        EditorAnalyticsUtils::BoolToStoredString(bIsVanilla)},
-			{EditorAnalyticsDefs::WasShutdownStoreKey,      EditorAnalyticsUtils::BoolToStoredString(bWasShutdown)},
-			{EditorAnalyticsDefs::IsUserLoggingOutStoreKey, EditorAnalyticsUtils::BoolToStoredString(bIsUserLoggingOut)},
-			{EditorAnalyticsDefs::IsInPIEStoreKey,          EditorAnalyticsUtils::BoolToStoredString(bIsInPIE)    },
-			{EditorAnalyticsDefs::IsInEnterpriseStoreKey,   EditorAnalyticsUtils::BoolToStoredString(bIsInEnterprise)},
-			{EditorAnalyticsDefs::IsInVRModeStoreKey,       EditorAnalyticsUtils::BoolToStoredString(bIsInVRMode)},
-			{EditorAnalyticsDefs::IsLowDriveSpaceStoreKey,  EditorAnalyticsUtils::BoolToStoredString(bIsLowDriveSpace)},
-			{EditorAnalyticsDefs::ProcessDiagnosticsStoreKey, FString::FromInt(ProcessDiagnostics)},
-		};
-
-		if (ExitCode.IsSet())
-		{
-			KeyValues.Emplace(EditorAnalyticsDefs::ExitCodeStoreKey, FString::FromInt(ExitCode.GetValue()));
-		}
-
-		if (MonitorExceptCode.IsSet())
-		{
-			KeyValues.Emplace(EditorAnalyticsDefs::MonitorExceptCodeStoreKey, FString::FromInt(MonitorExceptCode.GetValue()));
-		}
-
-		if (MonitorExitCode.IsSet())
-		{
-			KeyValues.Emplace(EditorAnalyticsDefs::MonitorExitCodeStoreKey, FString::FromInt(MonitorExitCode.GetValue()));
-		}
-		
-		if (DeathTimestamp.IsSet())
-		{
-			KeyValues.Emplace(EditorAnalyticsDefs::DeathTimestampStoreKey, EditorAnalyticsUtils::TimestampToString(DeathTimestamp.GetValue()));
-		}
-
-		FPlatformMisc::SetStoredValues(EditorAnalyticsDefs::StoreId, StorageLocation, KeyValues);
-	}
-
-	return true;
-}
-
-bool FEditorAnalyticsSession::Load(const FString& InSessionID)
-{
-	if (!ensure(IsLockedBy(FPlatformTLS::GetCurrentThreadId())))
-	{
-		return false;
-	}
-
-	EditorAnalyticsUtils::LoadInternal(*this, InSessionID);
-	bAlreadySaved = false;
-
-	return true;
-}
-
-bool FEditorAnalyticsSession::Delete() const
-{
-	if (!ensure(IsLockedBy(FPlatformTLS::GetCurrentThreadId())))
-	{
-		return false;
-	}
-
-	FString SectionName = EditorAnalyticsUtils::GetSessionStorageLocation(SessionId);
-
-	FPlatformMisc::DeleteStoredValue(EditorAnalyticsDefs::StoreId, SectionName, EditorAnalyticsDefs::AppIdStoreKey);
-	FPlatformMisc::DeleteStoredValue(EditorAnalyticsDefs::StoreId, SectionName, EditorAnalyticsDefs::AppVersionStoreKey);
-	FPlatformMisc::DeleteStoredValue(EditorAnalyticsDefs::StoreId, SectionName, EditorAnalyticsDefs::UserIdStoreKey);
-
-	FPlatformMisc::DeleteStoredValue(EditorAnalyticsDefs::StoreId, SectionName, EditorAnalyticsDefs::ProjectNameStoreKey);
-	FPlatformMisc::DeleteStoredValue(EditorAnalyticsDefs::StoreId, SectionName, EditorAnalyticsDefs::ProjectIDStoreKey);
-	FPlatformMisc::DeleteStoredValue(EditorAnalyticsDefs::StoreId, SectionName, EditorAnalyticsDefs::ProjectDescriptionStoreKey);
-	FPlatformMisc::DeleteStoredValue(EditorAnalyticsDefs::StoreId, SectionName, EditorAnalyticsDefs::ProjectVersionStoreKey);
-	FPlatformMisc::DeleteStoredValue(EditorAnalyticsDefs::StoreId, SectionName, EditorAnalyticsDefs::EngineVersionStoreKey);
-	FPlatformMisc::DeleteStoredValue(EditorAnalyticsDefs::StoreId, SectionName, EditorAnalyticsDefs::PlatformProcessIDStoreKey);
-	FPlatformMisc::DeleteStoredValue(EditorAnalyticsDefs::StoreId, SectionName, EditorAnalyticsDefs::MonitorProcessIDStoreKey);
-	FPlatformMisc::DeleteStoredValue(EditorAnalyticsDefs::StoreId, SectionName, EditorAnalyticsDefs::ExitCodeStoreKey);
-	FPlatformMisc::DeleteStoredValue(EditorAnalyticsDefs::StoreId, SectionName, EditorAnalyticsDefs::MonitorExceptCodeStoreKey);
-	FPlatformMisc::DeleteStoredValue(EditorAnalyticsDefs::StoreId, SectionName, EditorAnalyticsDefs::MonitorExitCodeStoreKey);
-
-	FPlatformMisc::DeleteStoredValue(EditorAnalyticsDefs::StoreId, SectionName, EditorAnalyticsDefs::StartupTimestampStoreKey);
-	FPlatformMisc::DeleteStoredValue(EditorAnalyticsDefs::StoreId, SectionName, EditorAnalyticsDefs::TimestampStoreKey);
-	FPlatformMisc::DeleteStoredValue(EditorAnalyticsDefs::StoreId, SectionName, EditorAnalyticsDefs::LastTickTimestampStoreKey);
-	FPlatformMisc::DeleteStoredValue(EditorAnalyticsDefs::StoreId, SectionName, EditorAnalyticsDefs::DeathTimestampStoreKey);
-	FPlatformMisc::DeleteStoredValue(EditorAnalyticsDefs::StoreId, SectionName, EditorAnalyticsDefs::SessionDurationStoreKey);
-	FPlatformMisc::DeleteStoredValue(EditorAnalyticsDefs::StoreId, SectionName, EditorAnalyticsDefs::Idle1MinStoreKey);
-	FPlatformMisc::DeleteStoredValue(EditorAnalyticsDefs::StoreId, SectionName, EditorAnalyticsDefs::Idle5MinStoreKey);
-	FPlatformMisc::DeleteStoredValue(EditorAnalyticsDefs::StoreId, SectionName, EditorAnalyticsDefs::Idle30MinStoreKey);
-	FPlatformMisc::DeleteStoredValue(EditorAnalyticsDefs::StoreId, SectionName, EditorAnalyticsDefs::TotalEditorInactivitySecondsStoreKey);
-	FPlatformMisc::DeleteStoredValue(EditorAnalyticsDefs::StoreId, SectionName, EditorAnalyticsDefs::CurrentUserActivityStoreKey);
-	FPlatformMisc::DeleteStoredValue(EditorAnalyticsDefs::StoreId, SectionName, EditorAnalyticsDefs::PluginsStoreKey);
-	FPlatformMisc::DeleteStoredValue(EditorAnalyticsDefs::StoreId, SectionName, EditorAnalyticsDefs::AverageFPSStoreKey);
-	FPlatformMisc::DeleteStoredValue(EditorAnalyticsDefs::StoreId, SectionName, EditorAnalyticsDefs::SessionTickCountStoreKey);
-	FPlatformMisc::DeleteStoredValue(EditorAnalyticsDefs::StoreId, SectionName, EditorAnalyticsDefs::EngineTickCountStoreKey);
-	FPlatformMisc::DeleteStoredValue(EditorAnalyticsDefs::StoreId, SectionName, EditorAnalyticsDefs::UserInteractionCountStoreKey);
-	FPlatformMisc::DeleteStoredValue(EditorAnalyticsDefs::StoreId, SectionName, EditorAnalyticsDefs::CommandLineStoreKey);
-
-	FPlatformMisc::DeleteStoredValue(EditorAnalyticsDefs::StoreId, SectionName, EditorAnalyticsDefs::TotalStallCountStoreKey);
-	FPlatformMisc::DeleteStoredValue(EditorAnalyticsDefs::StoreId, SectionName, EditorAnalyticsDefs::TotalStallReportedStoreKey);
-	FPlatformMisc::DeleteStoredValue(EditorAnalyticsDefs::StoreId, SectionName, EditorAnalyticsDefs::TopStallNameStoreKey);
-	FPlatformMisc::DeleteStoredValue(EditorAnalyticsDefs::StoreId, SectionName, EditorAnalyticsDefs::TopStallBudgetSecondsStoreKey);
-	FPlatformMisc::DeleteStoredValue(EditorAnalyticsDefs::StoreId, SectionName, EditorAnalyticsDefs::TopStallOverageSecondsStoreKey);
-	FPlatformMisc::DeleteStoredValue(EditorAnalyticsDefs::StoreId, SectionName, EditorAnalyticsDefs::TopStallTriggerCountStoreKey);
-
-	FPlatformMisc::DeleteStoredValue(EditorAnalyticsDefs::StoreId, SectionName, EditorAnalyticsDefs::DesktopGPUAdapterStoreKey);
-	FPlatformMisc::DeleteStoredValue(EditorAnalyticsDefs::StoreId, SectionName, EditorAnalyticsDefs::RenderingGPUAdapterStoreKey);
-	FPlatformMisc::DeleteStoredValue(EditorAnalyticsDefs::StoreId, SectionName, EditorAnalyticsDefs::GPUVendorIDStoreKey);
-	FPlatformMisc::DeleteStoredValue(EditorAnalyticsDefs::StoreId, SectionName, EditorAnalyticsDefs::GPUDeviceIDStoreKey);
-	FPlatformMisc::DeleteStoredValue(EditorAnalyticsDefs::StoreId, SectionName, EditorAnalyticsDefs::GRHIDeviceRevisionStoreKey);
-	FPlatformMisc::DeleteStoredValue(EditorAnalyticsDefs::StoreId, SectionName, EditorAnalyticsDefs::GRHIAdapterInternalDriverVersionStoreKey);
-	FPlatformMisc::DeleteStoredValue(EditorAnalyticsDefs::StoreId, SectionName, EditorAnalyticsDefs::GRHIAdapterUserDriverVersionStoreKey);
-	FPlatformMisc::DeleteStoredValue(EditorAnalyticsDefs::StoreId, SectionName, EditorAnalyticsDefs::GRHINameStoreKey);
-
-	FPlatformMisc::DeleteStoredValue(EditorAnalyticsDefs::StoreId, SectionName, EditorAnalyticsDefs::TotalPhysicalRAMStoreKey);
-	FPlatformMisc::DeleteStoredValue(EditorAnalyticsDefs::StoreId, SectionName, EditorAnalyticsDefs::CPUPhysicalCoresStoreKey);
-	FPlatformMisc::DeleteStoredValue(EditorAnalyticsDefs::StoreId, SectionName, EditorAnalyticsDefs::CPULogicalCoresStoreKey);
-	FPlatformMisc::DeleteStoredValue(EditorAnalyticsDefs::StoreId, SectionName, EditorAnalyticsDefs::CPUVendorStoreKey);
-	FPlatformMisc::DeleteStoredValue(EditorAnalyticsDefs::StoreId, SectionName, EditorAnalyticsDefs::CPUBrandStoreKey);
-
-	FPlatformMisc::DeleteStoredValue(EditorAnalyticsDefs::StoreId, SectionName, EditorAnalyticsDefs::OSMajorStoreKey);
-	FPlatformMisc::DeleteStoredValue(EditorAnalyticsDefs::StoreId, SectionName, EditorAnalyticsDefs::OSMinorStoreKey);
-	FPlatformMisc::DeleteStoredValue(EditorAnalyticsDefs::StoreId, SectionName, EditorAnalyticsDefs::OSVersionStoreKey);
-	FPlatformMisc::DeleteStoredValue(EditorAnalyticsDefs::StoreId, SectionName, EditorAnalyticsDefs::bIs64BitOSStoreKey);
-
-	FPlatformMisc::DeleteStoredValue(EditorAnalyticsDefs::StoreId, SectionName, EditorAnalyticsDefs::IsCrashStoreKey);
-	FPlatformMisc::DeleteStoredValue(EditorAnalyticsDefs::StoreId, SectionName, EditorAnalyticsDefs::IsGPUCrashStoreKey);
-	FPlatformMisc::DeleteStoredValue(EditorAnalyticsDefs::StoreId, SectionName, EditorAnalyticsDefs::IsDebuggerPresentStoreKey);
-	FPlatformMisc::DeleteStoredValue(EditorAnalyticsDefs::StoreId, SectionName, EditorAnalyticsDefs::IsDebuggerIgnoredStoreKey);
-	FPlatformMisc::DeleteStoredValue(EditorAnalyticsDefs::StoreId, SectionName, EditorAnalyticsDefs::WasDebuggerStoreKey);
-	FPlatformMisc::DeleteStoredValue(EditorAnalyticsDefs::StoreId, SectionName, EditorAnalyticsDefs::IsVanillaStoreKey);
-	FPlatformMisc::DeleteStoredValue(EditorAnalyticsDefs::StoreId, SectionName, EditorAnalyticsDefs::IsTerminatingStoreKey);
-	FPlatformMisc::DeleteStoredValue(EditorAnalyticsDefs::StoreId, SectionName, EditorAnalyticsDefs::WasShutdownStoreKey);
-	FPlatformMisc::DeleteStoredValue(EditorAnalyticsDefs::StoreId, SectionName, EditorAnalyticsDefs::IsUserLoggingOutStoreKey);
-	FPlatformMisc::DeleteStoredValue(EditorAnalyticsDefs::StoreId, SectionName, EditorAnalyticsDefs::IsInPIEStoreKey);
-	FPlatformMisc::DeleteStoredValue(EditorAnalyticsDefs::StoreId, SectionName, EditorAnalyticsDefs::IsInEnterpriseStoreKey);
-	FPlatformMisc::DeleteStoredValue(EditorAnalyticsDefs::StoreId, SectionName, EditorAnalyticsDefs::IsInVRModeStoreKey);
-	FPlatformMisc::DeleteStoredValue(EditorAnalyticsDefs::StoreId, SectionName, EditorAnalyticsDefs::IsLowDriveSpaceStoreKey);
-	FPlatformMisc::DeleteStoredValue(EditorAnalyticsDefs::StoreId, SectionName, EditorAnalyticsDefs::IsCrcExeMissingStoreKey);
-	FPlatformMisc::DeleteStoredValue(EditorAnalyticsDefs::StoreId, SectionName, EditorAnalyticsDefs::ProcessDiagnosticsStoreKey);
-
-	// Delete the log files.
-	EditorAnalyticsUtils::DeleteLogEvents(*this);
-
-	return true;
-}
-
-bool FEditorAnalyticsSession::GetStoredSessionIDs(TArray<FString>& OutSessions)
-{
-	if (!ensure(IsLockedBy(FPlatformTLS::GetCurrentThreadId())))
-	{
-		return false;
-	}
-
-	OutSessions = EditorAnalyticsUtils::GetSessionList();
-	return true;
-}
-
-bool FEditorAnalyticsSession::LoadAllStoredSessions(TArray<FEditorAnalyticsSession>& OutSessions)
-{
-	if (!ensure(IsLockedBy(FPlatformTLS::GetCurrentThreadId())))
-	{
-		return false;
-	}
-
-	TArray<FString> SessionIDs = EditorAnalyticsUtils::GetSessionList();
-
-	// Retrieve all the sessions in the list from storage
-	for (const FString& Id : SessionIDs)
-	{
-		FEditorAnalyticsSession NewSession;
-		EditorAnalyticsUtils::LoadInternal(NewSession, Id);
-
-		OutSessions.Add(MoveTemp(NewSession));
-	}
-
-	return true;
-}
-
-bool FEditorAnalyticsSession::SaveStoredSessionIDs(const TArray<FString>& InSessions)
+void CleanupDeprecatedAnalyticSessions(const FTimespan& MaxAge)
 {
-	// build up a new SessionList string
-	FString SessionListString;
-	for (const FString& Session : InSessions)
+	FSystemWideCriticalSection SysWideLock(EditorAnalyticsDefs::GlobalLockName, FTimespan::Zero());
+	if (!SysWideLock.IsValid())
 	{
-		if (!SessionListString.IsEmpty())
-		{
-			SessionListString.Append(TEXT(","));
-		}
-
-		SessionListString.Append(Session);
-	}
-
-	if (!ensure(IsLockedBy(FPlatformTLS::GetCurrentThreadId())))
-	{
-		return false;
-	}
-
-	FPlatformMisc::SetStoredValue(EditorAnalyticsDefs::StoreId, EditorAnalyticsDefs::SessionSummarySection, EditorAnalyticsDefs::SessionListStoreKey, SessionListString);
-	return true;
-}
-
-void FEditorAnalyticsSession::CleanupOutdatedIncompatibleSessions(const FTimespan& MaxAge)
-{
-	if (!ensure(IsLockedBy(FPlatformTLS::GetCurrentThreadId())))
-	{
-		return;
+		return; // Failed to lock, don't bother, this is just for cleaning old deprecated stuff, will do next time.
 	}
 
 	// Helper function to scan and clear sessions stored in sections corresponding to older versions.
@@ -872,7 +119,7 @@ void FEditorAnalyticsSession::CleanupOutdatedIncompatibleSessions(const FTimespa
 
 			for (const FString& SessionID : SessionIDs)
 			{
-				// All versions (1_0, 1_1 and 1_2) have a 'Timestamp' field. If it is not found, the session was partially deleted and should be cleaned up.
+				// All versions had a 'Timestamp' field. If it is not found, the session was partially deleted and should be cleaned up.
 				FString SessionSectionName = SectionVersion / SessionID;
 				FString TimestampStr;
 				if (FPlatformMisc::GetStoredValue(EditorAnalyticsDefs::StoreId, SessionSectionName, EditorAnalyticsDefs::TimestampStoreKey, TimestampStr))
@@ -883,6 +130,9 @@ void FEditorAnalyticsSession::CleanupOutdatedIncompatibleSessions(const FTimespa
 						// Don't delete the section yet, it contains a session young enough that could be sent if the user launch the Editor corresponding to this session format again.
 						return;
 					}
+
+					// Clean up the log events (if any) left-over by this session.
+					EditorAnalyticsUtils::DeleteLogEvents(SessionID);
 				}
 			}
 		}
@@ -891,96 +141,42 @@ void FEditorAnalyticsSession::CleanupOutdatedIncompatibleSessions(const FTimespa
 		FPlatformMisc::DeleteStoredSection(EditorAnalyticsDefs::StoreId, SectionVersion);
 	};
 
-	// Delete older and incompatible sections unless it contains a valid session young enough that would be picked up
-	// if an older Editor with compatible format was launched again.
-	for (int i = 0; i < UE_ARRAY_COUNT(EditorAnalyticsDefs::DeprecatedVersions); ++i)
-	{
-		CleanupVersionedSection(EditorAnalyticsDefs::DeprecatedVersions[i]);
-	}
-}
-
-void FEditorAnalyticsSession::LogEvent(EEventType InEventType, const FDateTime& InTimestamp)
-{
-	EditorAnalyticsUtils::LogSessionEvent(*this, InEventType, InTimestamp);
-}
-
-bool FEditorAnalyticsSession::FindSession(const uint32 InSessionProcessId, FEditorAnalyticsSession& OutSession)
-{
-	if (!ensure(IsLockedBy(FPlatformTLS::GetCurrentThreadId())))
+	if (IFileManager::Get().DirectoryExists(*EditorAnalyticsUtils::GetSessionEventLogDir()))
 	{
-		return false;
-	}
-
-	TArray<FString> SessionIDs = EditorAnalyticsUtils::GetSessionList();
+		int32 FileCount = 0;
+		// Find the 'log events' directory that could be left over from previous executions.
+		FRegexPattern Pattern(TEXT(R"((^[a-fA-F0-9-]+)_([0-9]+)_([0-9]+)_([0-9]+)_([0-9]+)_([0-9]+)_([0-9]+))")); // Need help with regex? Try https://regex101.com/
+		IFileManager::Get().IterateDirectory(*EditorAnalyticsUtils::GetSessionEventLogDir(), [&Pattern, &MaxAge, &FileCount](const TCHAR* Pathname, bool bIsDir)
+		{
+			++FileCount;
+			if (bIsDir) // Log events were encoded in the directory name.
+			{
+				FRegexMatcher Matcher(Pattern, FPaths::GetCleanFilename(Pathname));
+				if (Matcher.FindNext())
+				{
+					FFileStatData DirStats = IFileManager::Get().GetStatData(Pathname);
+					if ((FDateTime::UtcNow() - DirStats.CreationTime).GetTotalSeconds() >= MaxAge.GetTotalSeconds())
+					{
+						if (IFileManager::Get().DeleteDirectory(Pathname))
+						{
+							--FileCount;
+						}
+					}
+				}
+			}
+			return true; // Continue.
+		});
 
-	// Retrieve all the sessions in the list from storage
-	for (const FString& Id : SessionIDs)
-	{
-		FEditorAnalyticsSession Session;
-		EditorAnalyticsUtils::LoadInternal(Session, Id);
-		if (Session.PlatformProcessID == InSessionProcessId)
+		if (FileCount == 0)
 		{
-			OutSession = MoveTemp(Session);
-			return true;
+			IFileManager::Get().DeleteDirectory(*EditorAnalyticsUtils::GetSessionEventLogDir());
 		}
 	}
 
-	return false;
-}
-
-bool FEditorAnalyticsSession::SaveExitCode(int32 InExitCode, const FDateTime& ApproximativeEditorDeathTime)
-{
-	if (!ensure(IsLockedBy(FPlatformTLS::GetCurrentThreadId())))
-	{
-		return false;
-	}
-
-	ExitCode.Emplace(InExitCode);
-	DeathTimestamp.Emplace(ApproximativeEditorDeathTime);
-
-	const FString StorageLocation = EditorAnalyticsUtils::GetSessionStorageLocation(SessionId);
-	return FPlatformMisc::SetStoredValue(EditorAnalyticsDefs::StoreId, StorageLocation, EditorAnalyticsDefs::ExitCodeStoreKey, FString::FromInt(InExitCode))
-		&& FPlatformMisc::SetStoredValue(EditorAnalyticsDefs::StoreId, StorageLocation, EditorAnalyticsDefs::DeathTimestampStoreKey, EditorAnalyticsUtils::TimestampToString(ApproximativeEditorDeathTime));
-}
-
-bool FEditorAnalyticsSession::SaveMonitorExceptCode(int32 InExceptCode)
-{
-	if (!ensure(IsLockedBy(FPlatformTLS::GetCurrentThreadId())))
-	{
-		return false;
-	}
-
-	MonitorExceptCode.Emplace(InExceptCode);
-
-	const FString StorageLocation = EditorAnalyticsUtils::GetSessionStorageLocation(SessionId);
-	return FPlatformMisc::SetStoredValue(EditorAnalyticsDefs::StoreId, StorageLocation, EditorAnalyticsDefs::MonitorExceptCodeStoreKey, FString::FromInt(InExceptCode));
-}
-
-void FEditorAnalyticsSession::CreateMinimalCrashSession(const TOptional<int32>& ExitCode)
-{
-	if (!ensure(IsLockedBy(FPlatformTLS::GetCurrentThreadId()))) // Caller needs to own the session lock.
+	// Delete older and incompatible sections unless it contains a valid session young enough that would be picked up
+	// if an older Editor with compatible format was launched again.
+	for (int i = 0; i < UE_ARRAY_COUNT(EditorAnalyticsDefs::DeprecatedVersions); ++i)
 	{
-		return;
+		CleanupVersionedSection(EditorAnalyticsDefs::DeprecatedVersions[i]);
 	}
-
-	FEditorAnalyticsSession MinimalSession;
-	MinimalSession.SessionId = FGuid::NewGuid().ToString(EGuidFormats::DigitsWithHyphens); // To have a unique session id used as key in the storage.
-	MinimalSession.bCrashed = true;
-	MinimalSession.ExitCode = ExitCode;
-	MinimalSession.StartupTimestamp = FDateTime::UtcNow();
-	MinimalSession.Timestamp = MinimalSession.StartupTimestamp; // To enable session 'expiration'.
-	MinimalSession.EngineVersion = FEngineVersion::Current().ToString(EVersionComponent::Changelist); // To account the crash for the corresponding engine version.
-
-	MinimalSession.Save();
-
-	// Update the session list.
-	TArray<FString> StoredSessions;
-	GetStoredSessionIDs(StoredSessions);
-	StoredSessions.Add(MinimalSession.SessionId);
-	SaveStoredSessionIDs(StoredSessions);
-}
-
-bool FEditorAnalyticsSession::IsMinimalCrashSession() const
-{
-	return AppId == EditorAnalyticsDefs::UnknownAppIdString && ProjectID.IsEmpty() && Plugins.Num() == 0; // Check fields values that are not expected when the session is fully initialized.
 }
diff --git a/Engine/Source/Developer/EditorAnalyticsSession/Private/EditorSessionSummarySender.cpp b/Engine/Source/Developer/EditorAnalyticsSession/Private/EditorSessionSummarySender.cpp
--- a/Engine/Source/Developer/EditorAnalyticsSession/Private/EditorSessionSummarySender.cpp
+++ b/Engine/Source/Developer/EditorAnalyticsSession/Private/EditorSessionSummarySender.cpp
@@ -1,405 +1,3 @@
 // Copyright Epic Games, Inc. All Rights Reserved.
 
 #include "EditorSessionSummarySender.h"
-
-#include "AnalyticsEventAttribute.h"
-#include "Algo/Transform.h"
-#include "IAnalyticsProviderET.h"
-#include "EditorAnalyticsSession.h"
-#include "HAL/PlatformProcess.h"
-#include "Misc/EngineVersion.h"
-#include "Internationalization/Regex.h"
-#include "GenericPlatform/GenericPlatformCrashContext.h"
-
-DEFINE_LOG_CATEGORY_STATIC(LogEditorSessionSummary, Verbose, All);
-
-/* FEditorSessionSummarySender */
-
-namespace EditorSessionSenderDefs
-{
-	static const FTimespan SessionExpiration = FTimespan::FromDays(30.0);
-	static const float HeartbeatPeriodSeconds = 60;
-
-	// shutdown types
-	static const FString RunningSessionToken(TEXT("Running"));
-	static const FString ShutdownSessionToken(TEXT("Shutdown"));
-	static const FString CrashSessionToken(TEXT("Crashed"));
-	static const FString TerminatedSessionToken(TEXT("Terminated"));
-	static const FString DebuggerSessionToken(TEXT("Debugger"));
-	static const FString AbnormalSessionToken(TEXT("AbnormalShutdown"));
-}
-
-/** Contains the analytics session and some extra data to report. */
-struct FEditorAnalyticReportData
-{
-	/** The analytic session to report. */
-	FEditorAnalyticsSession Session;
-
-	/** Number of Editor instances (same engine version than Session) that crashed prior the analytic was initialized and wasn't reported yet. */
-	uint32 DelayedCrashCount = 0;
-
-	/** Number of Editor instances (same engine version than Session) that was killed by user prior the analytic was initialized and wasn't reported yet. */
-	uint32 DelayedKillCount = 0;
-
-	/* Number of Editor instance (same engine version than Session) that exited abnormally prior the analytic was initialized and wasn't reported yet. */
-	uint32 DelayedAbnormalCount = 0;
-};
-
-
-FEditorSessionSummarySender::FEditorSessionSummarySender(IAnalyticsProviderET& InAnalyticsProvider, const FString& InSenderName, const uint32 InCurrentSessionProcessId)
-	: HeartbeatTimeElapsed(0.0f)
-	, AnalyticsProvider(InAnalyticsProvider)
-	, Sender(InSenderName)
-	, CurrentSessionProcessId(InCurrentSessionProcessId)
-{
-}
-
-FEditorSessionSummarySender::~FEditorSessionSummarySender()
-{
-}
-
-void FEditorSessionSummarySender::Tick(float DeltaTime)
-{
-	HeartbeatTimeElapsed += DeltaTime;
-
-	if (HeartbeatTimeElapsed > EditorSessionSenderDefs::HeartbeatPeriodSeconds)
-	{
-		HeartbeatTimeElapsed = 0.0f;
-
-		SendStoredSessions();
-	}
-}
-
-void FEditorSessionSummarySender::Shutdown()
-{
-	SendStoredSessions(/*bForceSendCurrentSession*/true);
-}
-
-void FEditorSessionSummarySender::SetMonitorDiagnosticLogs(TMap<uint32, TTuple<FString, FDateTime>>&& Logs)
-{
-	MonitorMiniLogs = Logs;
-}
-
-void FEditorSessionSummarySender::SendStoredSessions(const bool bForceSendOwnedSession) const
-{
-	// Load the list of sessions to process. Expect contention on the analytic session lock between the Editor and CrashReportClientEditor (on windows) or between Editor instances (on mac/linux)
-	//   - Try every 'n' seconds if bForceSendCurrentSession is true.
-	//   - Don't block and don't loop if bForceSendCurrentSession is false.
-	TArray<FEditorAnalyticReportData> SessionsToReport;
-	FTimespan Timemout(bForceSendOwnedSession ? FTimespan::FromSeconds(5) : FTimespan::Zero());
-	bool bSessionsLoaded = false;
-	do
-	{
-		if (FEditorAnalyticsSession::Lock(Timemout))
-		{
-			// Get list of sessions in storage
-			TArray<FEditorAnalyticsSession> ExistingSessions;
-			TArray<FEditorAnalyticsSession> MinimalCrashSessions;
-			FEditorAnalyticsSession::LoadAllStoredSessions(ExistingSessions);
-
-			TArray<FEditorAnalyticsSession> SessionsToDelete;
-
-			// Whether this summary sender instance was tasked to send the specified session.
-			auto HasOwnershipOf = [this](const FEditorAnalyticsSession& InSession) { return InSession.PlatformProcessID == CurrentSessionProcessId; };
-
-			// Whether the process(es) writing and/or sending the specified session are all dead (and cannot write or send the session anymore).
-			auto IsOrphan = [](const FEditorAnalyticsSession& InSession) { return !FPlatformProcess::IsApplicationRunning(InSession.PlatformProcessID) && (InSession.MonitorProcessID == 0 || !FPlatformProcess::IsApplicationRunning(InSession.MonitorProcessID)); };
-
-			// Check each stored session to see if they should be sent or not
-			for (const FEditorAnalyticsSession& Session : ExistingSessions)
-			{
-				if (HasOwnershipOf(Session)) // This process was configured to send this session.
-				{
-					if (!bForceSendOwnedSession) // Don't send until forced to (on shutdown).
-					{
-						continue;
-					}
-				}
-				else if (!IsOrphan(Session))
-				{
-					continue; // Skip, another process is in charge of sending this session.
-				}
-
-				const FTimespan SessionAge = FDateTime::UtcNow() - Session.Timestamp;
-				if (SessionAge < EditorSessionSenderDefs::SessionExpiration)
-				{
-					if (Session.IsMinimalCrashSession())
-					{
-						MinimalCrashSessions.Add(Session);
-						continue; // Don't schedule this minimal session for reporting/deletiong yet, it's data must be piggybacked off a full and compatible session to be reported/deleted.
-					}
-
-					SessionsToReport.Add(FEditorAnalyticReportData{Session,0, 0, 0});
-				}
-				SessionsToDelete.Add(Session);
-			}
-
-			// Try to piggyback the minimal crash session information off a valid session. Minimal sessions are incomplete and will be rejected/ignored by the backend or may disrupt various analysis.
-			for (FEditorAnalyticsSession& MinimalCrashSession : MinimalCrashSessions)
-			{
-				// Piggyback information of the minimal crash session off a session that has the same engine version. (Stats are aggregated by engine version)
-				if (FEditorAnalyticReportData* ReportData = SessionsToReport.FindByPredicate([&MinimalCrashSession](const FEditorAnalyticReportData& Candidate) { return Candidate.Session.EngineVersion == MinimalCrashSession.EngineVersion; }))
-				{
-					// Account for the crashes that occurred before the analytic was initialized.
-					if (MinimalCrashSession.ExitCode.IsSet())
-					{
-						if (*MinimalCrashSession.ExitCode == 3) // 3 is the engine exit code when a crash is captured and reported.
-						{
-							ReportData->DelayedCrashCount += 1;
-						}
-						else if (PLATFORM_WINDOWS && *ReportData->Session.ExitCode == 1) // On windows, exit code 1 is usually because the user killed the Editor.
-						{
-							ReportData->DelayedKillCount += 1;
-						}
-						else // Any other exit codes.
-						{
-							ReportData->DelayedAbnormalCount += 1;
-						}
-					}
-					else // A minimal crash session was created, so the Editor crashed, but for some reasons, the exit code was unknown.
-					{
-						ReportData->DelayedAbnormalCount += 1;
-					}
-
-					// The minimal crash session was accounted for, schedule it for deletion.
-					SessionsToDelete.Add(MoveTemp(MinimalCrashSession));
-				}
-			}
-
-			// NOTE: The sessions are deleted before sending (while holding the lock) to prevent another process from finding and sending the same orphan sessions twice.
-			//       This also ensures a session will never be sent twice in case a crash occurs while sending, but all sessions loaded in memory for sending will also be lost.
-			for (const FEditorAnalyticsSession& ToDelete : SessionsToDelete)
-			{
-				ToDelete.Delete();
-				ExistingSessions.RemoveAll([&ToDelete](const FEditorAnalyticsSession& Session)
-					{
-						return Session.SessionId == ToDelete.SessionId;
-					});
-			}
-
-			TArray<FString> SessionIDs;
-			SessionIDs.Reserve(ExistingSessions.Num());
-			Algo::Transform(ExistingSessions, SessionIDs, &FEditorAnalyticsSession::SessionId);
-
-			FEditorAnalyticsSession::SaveStoredSessionIDs(SessionIDs);
-
-			// Trim left-over sessions that were written using an older format (so weren't loaded above) and that are now expired because the corresponding Editor version wasn't used recently.
-			FEditorAnalyticsSession::CleanupOutdatedIncompatibleSessions(EditorSessionSenderDefs::SessionExpiration);
-
-			FEditorAnalyticsSession::Unlock();
-			bSessionsLoaded = true;
-		}
-	} while (bForceSendOwnedSession && !bSessionsLoaded); // Retry until session are loaded if the sender is forced to send the current session.
-
-	// Send the sessions (without holding the lock).
-	for (const FEditorAnalyticReportData& ReportData : SessionsToReport)
-	{
-		SendSessionSummaryEvent(ReportData);
-	}
-}
-
-void FEditorSessionSummarySender::SendSessionSummaryEvent(const FEditorAnalyticReportData& ReportData) const
-{
-	const FEditorAnalyticsSession& Session = ReportData.Session;
-	FGuid SessionId;
-	FString SessionIdString = Session.SessionId;
-	if (FGuid::Parse(SessionIdString, SessionId))
-	{
-		// convert session guid to one with braces for sending to analytics
-		SessionIdString = SessionId.ToString(EGuidFormats::DigitsWithHyphensInBraces);
-	}
-
-	FString ShutdownTypeString = Session.bCrashed ? EditorSessionSenderDefs::CrashSessionToken :
-		(Session.bWasEverDebugger ? EditorSessionSenderDefs::DebuggerSessionToken :
-		(Session.bIsTerminating ? EditorSessionSenderDefs::TerminatedSessionToken :
-		(Session.bWasShutdown ? EditorSessionSenderDefs::ShutdownSessionToken : EditorSessionSenderDefs::AbnormalSessionToken)));
-
-	FString PluginsString = FString::Join(Session.Plugins, TEXT(","));
-
-	TArray<FAnalyticsEventAttribute> AnalyticsAttributes;
-
-	// Track which app/user is sending the summary (summary can be sent by another process (CrashReportClient) or another instance in case of crash/abnormal terminaison.
-	AnalyticsAttributes.Emplace(TEXT("SummarySenderAppId"), AnalyticsProvider.GetAppID());
-	AnalyticsAttributes.Emplace(TEXT("SummarySenderAppVersion"), AnalyticsProvider.GetAppVersion());
-	AnalyticsAttributes.Emplace(TEXT("SummarySenderEngineVersion"), FEngineVersion::Current().ToString(EVersionComponent::Changelist)); // Same as in EditorSessionSummaryWriter.cpp
-	AnalyticsAttributes.Emplace(TEXT("SummarySenderUserId"), AnalyticsProvider.GetUserID());
-	AnalyticsAttributes.Emplace(TEXT("SummarySenderSessionId"), AnalyticsProvider.GetSessionID()); // Not stripping the {} around the GUID like EditorSessionSummaryWriter does with SessionId.
-
-	AnalyticsAttributes.Emplace(TEXT("ProjectName"), Session.ProjectName);
-	AnalyticsAttributes.Emplace(TEXT("ProjectID"), Session.ProjectID);
-	AnalyticsAttributes.Emplace(TEXT("ProjectDescription"), Session.ProjectDescription);
-	AnalyticsAttributes.Emplace(TEXT("ProjectVersion"), Session.ProjectVersion);
-	AnalyticsAttributes.Emplace(TEXT("Platform"), FPlatformProperties::IniPlatformName());
-	AnalyticsAttributes.Emplace(TEXT("SessionId"), SessionIdString); // The provider is expected to add it as "SessionID" param in the HTTP request, but keep it for completness, because the formats are slightly different.
-	AnalyticsAttributes.Emplace(TEXT("EngineVersion"), Session.EngineVersion); // The provider is expected to add it as "AppVersion" param in the HTTP request, but keep it for completness, because the formats are slightly different.
-	AnalyticsAttributes.Emplace(TEXT("CommandLine"), Session.CommandLine);
-	AnalyticsAttributes.Emplace(TEXT("ShutdownType"), ShutdownTypeString);
-	AnalyticsAttributes.Emplace(TEXT("StartupTimestamp"), Session.StartupTimestamp.ToIso8601());
-	AnalyticsAttributes.Emplace(TEXT("Timestamp"), Session.Timestamp.ToIso8601());
-	AnalyticsAttributes.Emplace(TEXT("LastTickTimestamp"), Session.LastTickTimestamp.ToIso8601());
-	AnalyticsAttributes.Emplace(TEXT("SessionDurationWall"), FMath::FloorToInt(static_cast<float>((Session.Timestamp - Session.StartupTimestamp).GetTotalSeconds()))); // Session duration from system date/time.
-	AnalyticsAttributes.Emplace(TEXT("SessionDuration"), Session.SessionDuration);
-	AnalyticsAttributes.Emplace(TEXT("1MinIdle"), Session.Idle1Min);
-	AnalyticsAttributes.Emplace(TEXT("5MinIdle"), Session.Idle5Min);
-	AnalyticsAttributes.Emplace(TEXT("30MinIdle"), Session.Idle30Min);
-	AnalyticsAttributes.Emplace(TEXT("TotalEditorInactivitySecs"), Session.TotalEditorInactivitySeconds);
-	AnalyticsAttributes.Emplace(TEXT("CurrentUserActivity"), Session.CurrentUserActivity);
-	AnalyticsAttributes.Emplace(TEXT("AverageFPS"), Session.AverageFPS);
-	AnalyticsAttributes.Emplace(TEXT("SessionTickCount"), Session.SessionTickCount);
-	AnalyticsAttributes.Emplace(TEXT("EngineTickCount"), Session.EngineTickCount);
-	AnalyticsAttributes.Emplace(TEXT("UserInteractionCount"), Session.UserInteractionCount);
-	AnalyticsAttributes.Emplace(TEXT("TotalStallCount"), Session.TotalStallCount);
-	AnalyticsAttributes.Emplace(TEXT("TotalStallReported"), Session.TotalStallReported);
-	AnalyticsAttributes.Emplace(TEXT("TopStallName"), Session.TopStallName);
-	AnalyticsAttributes.Emplace(TEXT("TopStallBudgetSeconds"), Session.TopStallBudgetSeconds);
-	AnalyticsAttributes.Emplace(TEXT("TopStallOverageSeconds"), Session.TopStallOverageSeconds);
-	AnalyticsAttributes.Emplace(TEXT("TopStallTriggerCount"), Session.TopStallTriggerCount);
-	AnalyticsAttributes.Emplace(TEXT("Plugins"), PluginsString);
-	AnalyticsAttributes.Emplace(TEXT("DesktopGPUAdapter"), Session.DesktopGPUAdapter);
-	AnalyticsAttributes.Emplace(TEXT("RenderingGPUAdapter"), Session.RenderingGPUAdapter);
-	AnalyticsAttributes.Emplace(TEXT("GPUVendorID"), Session.GPUVendorID);
-	AnalyticsAttributes.Emplace(TEXT("GPUDeviceID"), Session.GPUDeviceID);
-	AnalyticsAttributes.Emplace(TEXT("GRHIDeviceRevision"), Session.GRHIDeviceRevision);
-	AnalyticsAttributes.Emplace(TEXT("GRHIAdapterInternalDriverVersion"), Session.GRHIAdapterInternalDriverVersion);
-	AnalyticsAttributes.Emplace(TEXT("GRHIAdapterUserDriverVersion"), Session.GRHIAdapterUserDriverVersion);
-	AnalyticsAttributes.Emplace(TEXT("GRHIName"), Session.GRHIName);
-	AnalyticsAttributes.Emplace(TEXT("TotalPhysicalRAM"), Session.TotalPhysicalRAM);
-	AnalyticsAttributes.Emplace(TEXT("CPUPhysicalCores"), Session.CPUPhysicalCores);
-	AnalyticsAttributes.Emplace(TEXT("CPULogicalCores"), Session.CPULogicalCores);
-	AnalyticsAttributes.Emplace(TEXT("CPUVendor"), Session.CPUVendor);
-	AnalyticsAttributes.Emplace(TEXT("CPUBrand"), Session.CPUBrand);
-	AnalyticsAttributes.Emplace(TEXT("OSMajor"), Session.OSMajor);
-	AnalyticsAttributes.Emplace(TEXT("OSMinor"), Session.OSMinor);
-	AnalyticsAttributes.Emplace(TEXT("OSVersion"), Session.OSVersion);
-	AnalyticsAttributes.Emplace(TEXT("Is64BitOS"), Session.bIs64BitOS);
-	AnalyticsAttributes.Emplace(TEXT("GPUCrash"), Session.bGPUCrashed);
-	AnalyticsAttributes.Emplace(TEXT("WasDebugged"), Session.bWasEverDebugger);
-	AnalyticsAttributes.Emplace(TEXT("WasDebuggerIgnored"), Session.bIsDebuggerIgnored);
-	AnalyticsAttributes.Emplace(TEXT("IsVanilla"), Session.bIsVanilla);
-	AnalyticsAttributes.Emplace(TEXT("WasShutdown"), Session.bWasShutdown);
-	AnalyticsAttributes.Emplace(TEXT("IsUserLoggingOut"), Session.bIsUserLoggingOut);
-	AnalyticsAttributes.Emplace(TEXT("IsInPIE"), Session.bIsInPIE);
-	AnalyticsAttributes.Emplace(TEXT("IsInEnterprise"), Session.bIsInEnterprise);
-	AnalyticsAttributes.Emplace(TEXT("IsInVRMode"), Session.bIsInVRMode);
-	AnalyticsAttributes.Emplace(TEXT("IsLowDriveSpace"), Session.bIsLowDriveSpace);
-	AnalyticsAttributes.Emplace(TEXT("IsCrcMissing"), Session.bIsCrcExeMissing);
-	AnalyticsAttributes.Emplace(TEXT("SentFrom"), Sender);
-	AnalyticsAttributes.Emplace(TEXT("MonitorPid"), Session.MonitorProcessID); // For out-of-process monitoring, if this is 0, this will mean CRC failed to launch or crashed very early.
-	AnalyticsAttributes.Emplace(TEXT("DelayedCrashCount"), ReportData.DelayedCrashCount); // Piggybacked information from 'minimal crash sessions', previous Editor instances that crashed before analytics was initialized.
-	AnalyticsAttributes.Emplace(TEXT("DelayedKillCount"), ReportData.DelayedKillCount);   // Piggybacked information from 'minimal crash sessions', previous Editor instances that were killed before analytics was initialized.
-	AnalyticsAttributes.Emplace(TEXT("DelayedAbnormalCount"), ReportData.DelayedAbnormalCount); // Piggybacked information from 'minimal crash sessions', previous Editor instances that exited abnormally before analytics was initialized.
-	AnalyticsAttributes.Emplace(TEXT("ProcessDiagnostics"), Session.ProcessDiagnostics); // Non-zero means that a diagnostic feature affecting performance was on. (like using a stomp allocator).
-
-	bool bShouldAttachMonitorLog = (ShutdownTypeString != EditorSessionSenderDefs::ShutdownSessionToken && ShutdownTypeString != EditorSessionSenderDefs::TerminatedSessionToken);
-
-	// Add the exit code to the report if it was set by out-of-process monitor (CrashReportClientEditor).
-	if (Session.ExitCode.IsSet())
-	{
-		AnalyticsAttributes.Emplace(TEXT("ExitCode"), Session.ExitCode.GetValue());
-		bShouldAttachMonitorLog |= Session.ExitCode.GetValue() != 0;
-	}
-	else
-	{
-		bShouldAttachMonitorLog = true; // Try to figure out why exit code could not be set.
-	}
-
-	if (Session.DeathTimestamp.IsSet())
-	{
-		AnalyticsAttributes.Emplace(TEXT("DeathTimestamp"), Session.DeathTimestamp->ToIso8601());
-	}
-
-	if (Session.MonitorExitCode.IsSet())
-	{
-		AnalyticsAttributes.Emplace(TEXT("MonitorExitCode"), Session.MonitorExitCode.GetValue());
-		bShouldAttachMonitorLog = true; // Try to figure out why CRC exited prematurely.
-	}
-
-	// Add the monitor exception code in case the out-of-process monitor (CrashReportClientEditor) crashed itself, caught the exception and was able to store it in the session before dying.
-	TOptional<int32> MonitorExceptCode = Session.MonitorExceptCode;
-	bShouldAttachMonitorLog |= MonitorExceptCode.IsSet();
-
-	// If the session did not end normally, try to attache the mini-log created for that session to help diagnose abnormal terminations.
-	if (bShouldAttachMonitorLog && !Session.bWasEverDebugger)
-	{
-		// Check if a monitor log is available. (Set by CrashReportClientEditor before sending the summary events).
-		if (const TTuple<FString, FDateTime>* MonitorLog = MonitorMiniLogs.Find(Session.MonitorProcessID))
-		{
-			AnalyticsAttributes.Emplace(TEXT("MonitorLog"), MonitorLog->Get<0>()); // CRC diagnostic mini-log.
-			AnalyticsAttributes.Emplace(TEXT("MonitorTimestamp"), MonitorLog->Get<1>().ToIso8601()); // Last time CRC timestamped the diagnostic mini-log, i.e. the approximative CRC death time.
-
-			// If no monitor exception code is set, check in the log if one exists. The exception may have occurred before the session was created or
-			// the out of process might not have been able to acquire the session lock.
-			if (!MonitorExceptCode.IsSet() || MonitorExceptCode.GetValue() == ECrashExitCodes::OutOfProcessReporterExitedUnexpectedly)
-			{
-				// Find the first entry in the log that match an exception reported like: "CRC/Crash:-1073741819"
-				FRegexPattern CrashPattern(TEXT(R"(CRC\/Crash:([-0-9]+).*)")); // Need help with regex? Try https://regex101.com/
-				FRegexMatcher CrashMatcher(CrashPattern, MonitorLog->Get<0>());
-				if (CrashMatcher.FindNext())
-				{
-					AnalyticsAttributes.Emplace(TEXT("MonitorExceptCode"), FCString::Atoi(*CrashMatcher.GetCaptureGroup(1))); // Report the first exception code found in the log.
-					MonitorExceptCode.Reset(); // Except code was added, prevent adding it again below.
-				}
-				else
-				{
-					// Check for 'CRC/Error'. Those are errors reported in log and the interesting one are the failed 'check()'. Normally CRC doesn't output errors.
-					FRegexPattern ErrorPattern(TEXT(R"(CRC\/Error)")); // Need help with regex? Try https://regex101.com/
-					FRegexMatcher ErrorMatcher(ErrorPattern, MonitorLog->Get<0>());
-					if (ErrorMatcher.FindNext())
-					{
-						AnalyticsAttributes.Emplace(TEXT("MonitorExceptCode"), static_cast<int32>(ECrashExitCodes::OutOfProcessReporterCheckFailed));
-						MonitorExceptCode.Reset(); // Except code was added, prevent adding it again below.
-					}
-				}
-			}
-
-			// Because lock contention, CRC doesn't always have time to save the exit code in the session when the user close the Editor and logoff very quickly, especially if the
-			// user has several Editors opened, but it usually have time to write it to the mini log. Try to parse it.
-			if (!Session.ExitCode.IsSet())
-			{
-				// Search for a pattern like: "Editor/ExitCode:0"
-				FRegexPattern ExitCodePattern(TEXT(R"(Editor\/ExitCode:([-0-9]+).*)")); // Need help with regex? Try https://regex101.com/
-				FRegexMatcher ExitCodeMatcher(ExitCodePattern, MonitorLog->Get<0>());
-				if (ExitCodeMatcher.FindNext())
-				{
-					AnalyticsAttributes.Emplace(TEXT("ExitCode"), FCString::Atoi(*ExitCodeMatcher.GetCaptureGroup(1)));
-				}
-			}
-		}
-	}
-
-	if (MonitorExceptCode.IsSet())
-	{
-		AnalyticsAttributes.Emplace(TEXT("MonitorExceptCode"), MonitorExceptCode.GetValue());
-	}
-
-	// Was this summary produced by another process than itself or the out-of-process monitor for that run?
-	AnalyticsAttributes.Emplace(TEXT("DelayedSend"), Session.PlatformProcessID != CurrentSessionProcessId);
-
-	// Sending the summary event of the current process analytic session?
-	if (AnalyticsProvider.GetSessionID().Contains(Session.SessionId)) // The string (GUID) returned by GetSessionID() is surrounded with braces like "{3FEA3232-...}" while Session.SessionId is not -> "3FEA3232-..."
-	{
-		AnalyticsProvider.RecordEvent(TEXT("SessionSummary"), AnalyticsAttributes);
-	}
-	else // The summary was created by another process/instance in a different session. (Ex: Editor sending a summary a prevoulsy crashed instance or CrashReportClientEditor sending it on behalf of the Editor)
-	{
-		// The provider sending a 'summary event' created by another instance/process must parametrize its post request 'as if' it was sent from the instance/session that created it (backend expectation).
-		// Create a new provider to avoid interfering with the current session events. (ex. if another thread sends telemetry at the same time, don't accidently tag it with the wrong SessionID, AppID, etc.).
-		TSharedPtr<IAnalyticsProviderET> TempSummaryProvider = FAnalyticsET::Get().CreateAnalyticsProvider(AnalyticsProvider.GetConfig());
-		
-		// Reconfigure the analytics provider to sent the summary event 'as if' it was sent by the process that created it. This is required by the analytics backend.
-		FGuid SessionGuid;
-		FGuid::Parse(Session.SessionId, SessionGuid);
-		TempSummaryProvider->SetSessionID(SessionGuid.ToString(EGuidFormats::DigitsWithHyphensInBraces)); // Ensure to put back the {} around the GUID.
-		TempSummaryProvider->SetAppID(CopyTemp(Session.AppId));
-		TempSummaryProvider->SetAppVersion(CopyTemp(Session.AppVersion));
-		TempSummaryProvider->SetUserID(CopyTemp(Session.UserId));
-
-		// Send the summary.
-		TempSummaryProvider->RecordEvent(TEXT("SessionSummary"), AnalyticsAttributes);
-
-		// The temporary provider is about to be deleted (going out of scope), ensure it sents its report.
-		TempSummaryProvider->BlockUntilFlushed(2.0f);
-	}
-
-	UE_LOG(LogEditorSessionSummary, Log, TEXT("EditorSessionSummary sent report. Type=%s, SessionId=%s"), *ShutdownTypeString, *SessionIdString);
-}
diff --git a/Engine/Source/Developer/EditorAnalyticsSession/Public/EditorAnalyticsSession.h b/Engine/Source/Developer/EditorAnalyticsSession/Public/EditorAnalyticsSession.h
--- a/Engine/Source/Developer/EditorAnalyticsSession/Public/EditorAnalyticsSession.h
+++ b/Engine/Source/Developer/EditorAnalyticsSession/Public/EditorAnalyticsSession.h
@@ -2,242 +2,19 @@
 
 #pragma once
 
-#include "Containers/UnrealString.h"
-#include "HAL/CriticalSection.h"
-#include "Misc/DateTime.h"
+#include "Misc/CoreMiscDefines.h"
+#include "Misc/Timespan.h"
 #include "Modules/ModuleInterface.h"
-#include "Templates/Atomic.h"
-#include <atomic>
 
+struct UE_DEPRECATED(5.0, "Deprecated. This class was replaced by FEngineAnalyticsSessionSummary/FEditorAnalyticsSessionSummary") FEditorAnalyticsSession;
 struct EDITORANALYTICSSESSION_API FEditorAnalyticsSession
 {
-	enum class EEventType : int32
-	{
-		Crashed = 0,
-		GpuCrashed,
-		Terminated,
-		Shutdown,
-		LogOut,
-	};
-
-	FString SessionId;
-
-	FString AppId;
-	FString AppVersion;
-	FString UserId;
-
-	FString ProjectName;
-	FString ProjectID;
-	FString ProjectDescription;
-	FString ProjectVersion;
-	FString EngineVersion;
-	uint32 PlatformProcessID = 0;
-	uint32 MonitorProcessID = 0; // Set to the CrashReportClientEditor PID when out-of-process reporting is used.
-	TOptional<int32> ExitCode; // Set by CrashReportClientEditor after the Editor process exit when out-of-process reporting is used and reading the exit code is supported.
-	TOptional<int32> MonitorExceptCode; // Set in CrashReportClientEditor process when an exception or an error is caught.
-	TOptional<int32> MonitorExitCode; // Set in the Editor process when the Editor detects that CrashReportClientEditor process unexpectedly died.
-
-	FDateTime StartupTimestamp; // Wall time (UTC) when the session started.
-	FDateTime Timestamp; // Wall time (UTC) when the session was ended.
-	FDateTime LastTickTimestamp; // Wall time (UTC) of the last engine tick recorded for the session.
-	TOptional<FDateTime> DeathTimestamp; // Wall time (UTC) of the Editor death from CRC p.o.v. (when CRC observed the Editor death)
-	volatile int32 SessionDuration = 0; // The session duration in seconds, computed using FPlatformTime::Seconds() rather than Timestamp - StartupTimestamp which can be affected by daylight saving.
-	volatile int32 IdleSeconds = 0; // Can be updated from concurrent threads.
-	volatile int32 Idle1Min = 0;
-	volatile int32 Idle5Min = 0;
-	volatile int32 Idle30Min = 0;
-	volatile int32 TotalEditorInactivitySeconds = 0; // Account for user input and Editor process CPU usage. Add up gaps where the CPU was not used intensively and the user did not interact.
-	FString CurrentUserActivity;
-	TArray<FString> Plugins;
-	float AverageFPS = 0.0f;
-
-	uint64 SessionTickCount = 0; // Number of times the analytic session was ticked. Zero is the interesting value. If the Editor is hang during boot, some users may be prompt to kill it.
-	uint64 EngineTickCount = 0;  // Number or times the engine was ticked.
-	uint32 UserInteractionCount = 0; // Number of slate user interactions. Zero is the interesting value. If the Editor UI hang at start up, some users may be prompt to kill it.
-
-	uint64 TotalStallCount = 0;    // The total number of stalls detected, across every detector
-	uint64 TotalStallReported = 0; // The total number of stalls reported (by default we only report the first instance, or skip reporting some detectors entirely)
-	FString TopStallName;
-	double TopStallBudgetSeconds = 0.0;
-	double TopStallOverageSeconds = 0.0;
-	uint64 TopStallTriggerCount = 0;
-
-	FString DesktopGPUAdapter;
-	FString RenderingGPUAdapter;
-	uint32 GPUVendorID = 0;
-	uint32 GPUDeviceID = 0;
-	uint32 GRHIDeviceRevision = 0;
-	FString GRHIAdapterInternalDriverVersion;
-	FString GRHIAdapterUserDriverVersion;
-	FString GRHIName;
-
-	uint64 TotalPhysicalRAM = 0;
-	int32 CPUPhysicalCores = 0;
-	int32 CPULogicalCores = 0;
-	FString CPUVendor;
-	FString CPUBrand;
-
-	FString CommandLine;
-	FString OSMajor;
-	FString OSMinor;
-	FString OSVersion;
-
-	bool bIs64BitOS : 1;
-	bool bCrashed : 1;
-	bool bGPUCrashed : 1;
-	bool bIsDebuggerPresent : 1;
-	bool bWasEverDebugger : 1;
-	bool bIsVanilla : 1;
-	bool bIsTerminating : 1;
-	bool bWasShutdown : 1;
-	bool bIsUserLoggingOut: 1; // Also cover shutdown/reboot as logging out is part of the process. Logging out currently end up as an abnormal termination.
-	bool bIsInPIE : 1;
-	bool bIsInEnterprise : 1;
-	bool bIsInVRMode : 1;
-	bool bIsLowDriveSpace : 1;
-	bool bIsCrcExeMissing: 1; // CrashReportClient executable is missing? To explain with MonitorProcessID would be zero.
-	bool bIsDebuggerIgnored: 1; // True if GIgnoreDebugger is true.
-
-	uint32 ProcessDiagnostics;
-
-	FEditorAnalyticsSession();
-
-	/** 
-	 * Save this session to stored values.
-	 * @returns true if the session was successfully saved.
-	 */
-	bool Save();
-
-	/**
-	 *  Load a session with the given session ID from stored values.
-	 * @returns true if the session was found and successfully loaded.
-	 */
-	bool Load(const FString& InSessionID);
-
-	/**
-	 * Delete the stored values of this session.
-	 * Does not modify the actual session object.
-	 * @returns true if the session was successfully deleted.
-	 */
-	bool Delete() const;
-
-	/**
-	 * Retrieve a list of session IDs that are currently stored locally.
-	 * @returns true if the session IDs were successfully retrieved.
-	 */
-	static bool GetStoredSessionIDs(TArray<FString>& OutSessions);
-
-	/**
-	 * Read all stored sessions into the given array. This function only loads the sessions that are
-	 * compatible with the current session format. Between releases, the format of the stored session
-	 * can change. A newer engine will load sessions from a previous engine as long as the session format
-	 * did not change.
-	 * @returns true if the sessions were successfully loaded.
-	 * @see CleanupOutdatedIncompatibleSessions()
-	 */
-	static bool LoadAllStoredSessions(TArray<FEditorAnalyticsSession>& OutSessions);
-
-	/**
-	 * Save the given session IDs to storage.
-	 * @returns true if the session IDs were successfully saved.
-	 */
-	static bool SaveStoredSessionIDs(const TArray<FString>& InSessions);
-
-	/**
-	 * Delete sessions that are incompatible with the current session format and too old to be sent even if the Editor
-	 * version corresponding to the format was used again. This is a maintenance function to prevent accumulating dead
-	 * sessions over years.
-	 * @param MaxAge Threshold to delete sessions that are older than MaxAge. Incompatible sessions are kept around for
-	 *               some time because many users continue using older versions of the Editor that would be able to send
-	 *               those sessions that a newer Editor version wouldn't.
-	 */
-	static void CleanupOutdatedIncompatibleSessions(const FTimespan& MaxAge);
-
-	/**
-	 * Try to acquire the local storage lock without blocking.
-	 * @return true if the lock was acquired successfully.
-	 */
-	static bool TryLock() { return Lock(FTimespan::Zero()); }
-
-	/**
-	 * Acquire a lock for local storage.
-	 * @returns true if the lock was acquired successfully.
-	 */
-	static bool Lock(FTimespan Timeout = FTimespan::Zero());
-
-	/**
-	 * Unlock the local storage.
-	 */
-	static void Unlock();
-
-	/**
-	 * Append an event to the session log. The function is meant to record concurrent events, especially during a crash
-	 * with minimum contention. The logger appends and persists the events of interest locklessly on spot as opposed to
-	 * overriding existing values in the key-store. Appending is better because it prevent dealing with event ordering
-	 * on the spot (no synchronization needed) and will preserve more information.
-	 *
-	 * @note They key-store is not easily usable in a 'lockless' fashion. On Windows, the OS provides thread safe API
-	 *       to modify the registry (add/update). On Mac/Linux, the key-store is a simple file and without synchronization,
-	 *       concurrent writes will likely corrupt the file.
-	 */
-	void LogEvent(EEventType EventTpe, const FDateTime& Timestamp);
-
-	/**
-	 * Find the sessions for which the PlatformProcessID matches the specified session process id.
-	 * @param InSessionProcessId The session with this process ID.
-	 * @param[out] The found session if any (undefined if not found)
-	 * @return True if a session for the specified process id was found, false otherwise.
-	 */
-	static bool FindSession(const uint32 InSessionProcessId, FEditorAnalyticsSession& OutSession);
-
-	/**
-	 * Persist the Editor exit code corresponding of this session in the key-store value.
-	 * @return true if the exit code is persisted with the session data.
-	 * @node This function should only be called by the out of process monitor when the process exit value is known.
-	 */
-	bool SaveExitCode(int32 ExitCode, const FDateTime& ApproximativeDeathTimeUtc);
-
-	/**
-	 * Set the exception code that caused the monitor application to crash. When the monitoring application raises
-	 * and catches an unexpected exception, it tries to save the exception code it in the session before dying.
-	 * @return true if the exception code is persisted with the session data.
-	 * @note This is to diagnose the cases where CrashReportClientEditor send the summary event delayed and without
-	 *       the Editor exit code.
-	 */
-	bool SaveMonitorExceptCode(int32 ExceptCode);
-
-	/**
-	 * Create a very minimal session for an Editor crash that occurred before the analytic session was created. The
-	 * session lifecyle is the same as a normal session, but it cannot be sent 'as-it' to the backend because it is
-	 * incomplete. Instead, the session information is pulled and piggybacked with another valid session. Usually, a
-	 * minimal session will be created by CrashReportClient (in monitor mode) if the Editor crashed and the corresponding
-	 * analytic session couldn't be found.
-	 * @note The purpose of this functionality is to allow counting very early crashes that aren't reported to analytics otherwise.
-	 */
-	static void CreateMinimalCrashSession(const TOptional<int32>& ExitCode);
-	
-	/** Returns true if the session is minimal, i.e. created with CreateMinimalCrashSession(). */
-	bool IsMinimalCrashSession() const;
-
-private:
-	/** Returns true if the local storage is locked by the specified thread. */
-	static bool IsLockedBy(uint32 ThreadId);
-
-	/** Ensure exclusive access to the stored values. */
-	static TUniquePtr<FSystemWideCriticalSection> StoredValuesLock;
-
-	/** Keep track of the lock owner to prevent using FSystemWideCriticalSection recursively, which is not supported by the implementation. */
-	static std::atomic<uint64> StoredValuesLockOwnerInfo;
-
-	/** 
-	 * Has this session already been saved? 
-	 * If not, then the first save will write out session invariant details such as hardware specs.
-	 */
-	bool bAlreadySaved : 1;
 };
 
+/** This function removes expired sessions that were created by the deprecated system. */
+EDITORANALYTICSSESSION_API void CleanupDeprecatedAnalyticSessions(const FTimespan& MaxAge);
+
 
 class FEditorAnalyticsSessionModule : public IModuleInterface
 {
-
-};
\ No newline at end of file
+};
diff --git a/Engine/Source/Developer/EditorAnalyticsSession/Public/EditorSessionSummarySender.h b/Engine/Source/Developer/EditorAnalyticsSession/Public/EditorSessionSummarySender.h
--- a/Engine/Source/Developer/EditorAnalyticsSession/Public/EditorSessionSummarySender.h
+++ b/Engine/Source/Developer/EditorAnalyticsSession/Public/EditorSessionSummarySender.h
@@ -2,41 +2,9 @@
 
 #pragma once
 
-#include "Containers/UnrealString.h"
-#include "HAL/Platform.h"
-#include "Misc/DateTime.h"
-#include "Misc/Optional.h"
-#include "Containers/Map.h"
-#include "Templates/Tuple.h"
+#include "Misc/CoreMiscDefines.h"
 
-struct FEditorAnalyticsSession;
-class IAnalyticsProviderET;
-struct FEditorAnalyticReportData;
-
-/**
-  * Sender of SessionSummary events from all editor sessions in-between runs.
-  * Separated from Writer to make it easier to run it out-of-process.
-  */
+class UE_DEPRECATED(5.0, "Deprecated. This class was replaced by FAnalyticsSessionSummarySender") FEditorSessionSummarySender;
 class EDITORANALYTICSSESSION_API FEditorSessionSummarySender
 {
-public:
-	FEditorSessionSummarySender(IAnalyticsProviderET& InAnalyticsProvider, const FString& InSenderName, const uint32 InCurrentSessionProcessId);
-	~FEditorSessionSummarySender();
-
-	void Tick(float DeltaTime);
-	void Shutdown();
-
-	void SetMonitorDiagnosticLogs(TMap<uint32, TTuple<FString, FDateTime>>&& Logs);
-
-private:
-	/** Send any stored Sessions. */
-	void SendStoredSessions(const bool bForceSendCurrentSession = false) const;
-	void SendSessionSummaryEvent(const FEditorAnalyticReportData& ReportData) const;
-
-private:
-	float HeartbeatTimeElapsed;
-	IAnalyticsProviderET& AnalyticsProvider;
-	FString Sender;
-	uint32 CurrentSessionProcessId;
-	TMap<uint32, TTuple<FString,FDateTime>> MonitorMiniLogs; // Maps Monitor Process ID/Monitor Log
 };
diff --git a/Engine/Source/Runtime/Analytics/AnalyticsET/Private/AnalyticsSessionSummaryManager.cpp b/Engine/Source/Runtime/Analytics/AnalyticsET/Private/AnalyticsSessionSummaryManager.cpp
--- a/Engine/Source/Runtime/Analytics/AnalyticsET/Private/AnalyticsSessionSummaryManager.cpp
+++ b/Engine/Source/Runtime/Analytics/AnalyticsET/Private/AnalyticsSessionSummaryManager.cpp
@@ -51,7 +51,6 @@ namespace AnalyticsSessionUtils
 {
 	static const FTimespan OrphanSessionCheckPeriod = FTimespan::FromMinutes(5);
 	static const FTimespan ProcessGroupDeathGracePeriod = FTimespan::FromHours(1);
-	static const FTimespan ProcessGroupExpiration = FTimespan::FromDays(30.0);
 
 	static const FString OrphanSessionsOwnerCriticalSectionName = TEXT("UE_AnalyticsSessionSummaryManager_OrphanOwner");
 	static const FString OrphanSessionsOwnerFilename            = TEXT("8E1D46DBC38F4A789939D781E1B91520"); // A randomly generated GUID.
@@ -413,7 +412,7 @@ void FAnalyticsSessionSummaryManager::ProcessSummary(const FString& InProcessGro
 			}
 
 			// If the principal file is expired, expire the entire group.
-			if (InProcessGroup.PrincipalProcessId == PropertyFileInfo.ProcessId && FileAgeSecs > AnalyticsSessionUtils::ProcessGroupExpiration.GetTotalSeconds())
+			if (InProcessGroup.PrincipalProcessId == PropertyFileInfo.ProcessId && FileAgeSecs > GetSessionExpirationAge().GetTotalSeconds())
 			{
 				bExpired = true;
 				break;
@@ -608,7 +607,9 @@ void FAnalyticsSessionSummaryManager::CleanupExpiredFiles(const FString& SavedDi
 		{
 			// Time since the last file modification.
 			double FileAgeSecs = IFileManager::Get().GetFileAgeSeconds(Pathname);
-			if (FileAgeSecs > AnalyticsSessionUtils::ProcessGroupExpiration.GetTotalSeconds() + FTimespan::FromDays(1.0).GetTotalSeconds())
+
+			// Add up an extra day to the normal expiration delay so that we never clean up files that are on the edge.
+			if (FileAgeSecs > GetSessionExpirationAge().GetTotalSeconds() + FTimespan::FromDays(1.0).GetTotalSeconds())
 			{
 				IFileManager::Get().Delete(Pathname);
 			}
diff --git a/Engine/Source/Runtime/Analytics/AnalyticsET/Public/AnalyticsSessionSummaryManager.h b/Engine/Source/Runtime/Analytics/AnalyticsET/Public/AnalyticsSessionSummaryManager.h
--- a/Engine/Source/Runtime/Analytics/AnalyticsET/Public/AnalyticsSessionSummaryManager.h
+++ b/Engine/Source/Runtime/Analytics/AnalyticsET/Public/AnalyticsSessionSummaryManager.h
@@ -5,6 +5,7 @@
 #include "CoreMinimal.h"
 #include "Interfaces/IAnalyticsPropertyStore.h"
 #include "AnalyticsEventAttribute.h"
+#include "Misc/Timespan.h"
 
 class IAnalyticsSessionSummarySender;
 
@@ -111,6 +112,12 @@ public:
 	 */
 	static void CleanupExpiredFiles(const FString& SavedDir = TEXT(""));
 
+	/**
+	 * Returns the age at which a session is considered expired and shouldn't be sent anymore.
+	 * @return The expiration delay.
+	 */
+	static FTimespan GetSessionExpirationAge() { return FTimespan::FromDays(30); }
+
 private:
 	/** Information about a property store used by a process. */
 	struct FPropertyFileInfo
diff --git a/Engine/Source/Runtime/Engine/Private/Analytics/EditorSessionSummaryWriter.cpp b/Engine/Source/Runtime/Engine/Private/Analytics/EditorSessionSummaryWriter.cpp
deleted file mode 100644
diff --git a/Engine/Source/Runtime/Engine/Private/Analytics/EditorSessionSummaryWriter.h b/Engine/Source/Runtime/Engine/Private/Analytics/EditorSessionSummaryWriter.h
deleted file mode 100644
diff --git a/Engine/Source/Runtime/Engine/Private/EngineAnalytics.cpp b/Engine/Source/Runtime/Engine/Private/EngineAnalytics.cpp
--- a/Engine/Source/Runtime/Engine/Private/EngineAnalytics.cpp
+++ b/Engine/Source/Runtime/Engine/Private/EngineAnalytics.cpp
@@ -202,11 +202,7 @@ void FEngineAnalytics::Shutdown(bool bIsEngineShutdown)
 	}
 
 	// Clean up the outdated sessions created by the deprecated system. If an older compabile Editor is launched, it can still send them up before they get expired.
-	if (FEditorAnalyticsSession::TryLock())
-	{
-		FEditorAnalyticsSession::CleanupOutdatedIncompatibleSessions(FTimespan::FromDays(30));
-		FEditorAnalyticsSession::Unlock();
-	}
+	CleanupDeprecatedAnalyticSessions(FAnalyticsSessionSummaryManager::GetSessionExpirationAge());
 #endif
 
 	bIsInitialized = false;
