Fixed static analysis reporting inconsistent processor state from CoreMiscDefines.h.

diff --git a/Engine/Source/Developer/EditorAnalyticsSession/Private/EditorAnalyticsSession.cpp b/Engine/Source/Developer/EditorAnalyticsSession/Private/EditorAnalyticsSession.cpp
--- a/Engine/Source/Developer/EditorAnalyticsSession/Private/EditorAnalyticsSession.cpp
+++ b/Engine/Source/Developer/EditorAnalyticsSession/Private/EditorAnalyticsSession.cpp
@@ -1,7 +1,6 @@
 // Copyright Epic Games, Inc. All Rights Reserved.
 
 #include "EditorAnalyticsSession.h"
-#include "CoreMinimal.h"
 #include "Modules/ModuleManager.h"
 #include "Internationalization/Regex.h"
 #include "Containers/UnrealString.h"
diff --git a/Engine/Source/Developer/EditorAnalyticsSession/Public/EditorAnalyticsSession.h b/Engine/Source/Developer/EditorAnalyticsSession/Public/EditorAnalyticsSession.h
--- a/Engine/Source/Developer/EditorAnalyticsSession/Public/EditorAnalyticsSession.h
+++ b/Engine/Source/Developer/EditorAnalyticsSession/Public/EditorAnalyticsSession.h
@@ -2,7 +2,7 @@
 
 #pragma once
 
-#include "Misc/CoreMiscDefines.h"
+#include "CoreMinimal.h"
 #include "Misc/Timespan.h"
 #include "Modules/ModuleInterface.h"
 
diff --git a/Engine/Source/Developer/EditorAnalyticsSession/Public/EditorSessionSummarySender.h b/Engine/Source/Developer/EditorAnalyticsSession/Public/EditorSessionSummarySender.h
--- a/Engine/Source/Developer/EditorAnalyticsSession/Public/EditorSessionSummarySender.h
+++ b/Engine/Source/Developer/EditorAnalyticsSession/Public/EditorSessionSummarySender.h
@@ -2,7 +2,7 @@
 
 #pragma once
 
-#include "Misc/CoreMiscDefines.h"
+#include "CoreMinimal.h"
 
 class UE_DEPRECATED(5.0, "Deprecated. This class was replaced by FAnalyticsSessionSummarySender") FEditorSessionSummarySender;
 class EDITORANALYTICSSESSION_API FEditorSessionSummarySender
